- Moved GameplayDebugger from Developer/ into Runtime/ and Editor source folder to be able to use it in shipping build.
- Moved the code originally in a Editor/ folder into its own GameplayDebuggerEditor module and added it so it get dynamically loaded in the editor
- Added some delegates so GameplayDebugger code can communicate to GameplayDebuggerEditor without referencing it. As only GameplayDebuggerEditor should reference GameplayDebugger, not the opposite.

diff --git a/Engine/Source/Developer/GameplayDebugger/GameplayDebugger.Build.cs b/Engine/Source/Developer/GameplayDebugger/GameplayDebugger.Build.cs
deleted file mode 100644
diff --git a/Engine/Source/Developer/GameplayDebugger/Private/Editor/GameplayDebuggerCategoryConfigCustomization.cpp b/Engine/Source/Developer/GameplayDebugger/Private/Editor/GameplayDebuggerCategoryConfigCustomization.cpp
deleted file mode 100644
diff --git a/Engine/Source/Developer/GameplayDebugger/Private/Editor/GameplayDebuggerCategoryConfigCustomization.h b/Engine/Source/Developer/GameplayDebugger/Private/Editor/GameplayDebuggerCategoryConfigCustomization.h
deleted file mode 100644
diff --git a/Engine/Source/Developer/GameplayDebugger/Private/Editor/GameplayDebuggerEdMode.cpp b/Engine/Source/Developer/GameplayDebugger/Private/Editor/GameplayDebuggerEdMode.cpp
deleted file mode 100644
diff --git a/Engine/Source/Developer/GameplayDebugger/Private/Editor/GameplayDebuggerEdMode.h b/Engine/Source/Developer/GameplayDebugger/Private/Editor/GameplayDebuggerEdMode.h
deleted file mode 100644
diff --git a/Engine/Source/Developer/GameplayDebugger/Private/Editor/GameplayDebuggerExtensionConfigCustomization.cpp b/Engine/Source/Developer/GameplayDebugger/Private/Editor/GameplayDebuggerExtensionConfigCustomization.cpp
deleted file mode 100644
diff --git a/Engine/Source/Developer/GameplayDebugger/Private/Editor/GameplayDebuggerExtensionConfigCustomization.h b/Engine/Source/Developer/GameplayDebugger/Private/Editor/GameplayDebuggerExtensionConfigCustomization.h
deleted file mode 100644
diff --git a/Engine/Source/Developer/GameplayDebugger/Private/Editor/GameplayDebuggerInputConfigCustomization.cpp b/Engine/Source/Developer/GameplayDebugger/Private/Editor/GameplayDebuggerInputConfigCustomization.cpp
deleted file mode 100644
diff --git a/Engine/Source/Developer/GameplayDebugger/Private/Editor/GameplayDebuggerInputConfigCustomization.h b/Engine/Source/Developer/GameplayDebugger/Private/Editor/GameplayDebuggerInputConfigCustomization.h
deleted file mode 100644
diff --git a/Engine/Source/Developer/GameplayDebugger/Private/Editor/GameplayDebuggerToolkit.cpp b/Engine/Source/Developer/GameplayDebugger/Private/Editor/GameplayDebuggerToolkit.cpp
deleted file mode 100644
diff --git a/Engine/Source/Developer/GameplayDebugger/Private/Editor/GameplayDebuggerToolkit.h b/Engine/Source/Developer/GameplayDebugger/Private/Editor/GameplayDebuggerToolkit.h
deleted file mode 100644
diff --git a/Engine/Source/Developer/GameplayDebugger/Private/GameplayDebuggerAddonBase.cpp b/Engine/Source/Developer/GameplayDebugger/Private/GameplayDebuggerAddonBase.cpp
deleted file mode 100644
diff --git a/Engine/Source/Developer/GameplayDebugger/Private/GameplayDebuggerAddonManager.cpp b/Engine/Source/Developer/GameplayDebugger/Private/GameplayDebuggerAddonManager.cpp
deleted file mode 100644
diff --git a/Engine/Source/Developer/GameplayDebugger/Private/GameplayDebuggerAddonManager.h b/Engine/Source/Developer/GameplayDebugger/Private/GameplayDebuggerAddonManager.h
deleted file mode 100644
diff --git a/Engine/Source/Developer/GameplayDebugger/Private/GameplayDebuggerCategory.cpp b/Engine/Source/Developer/GameplayDebugger/Private/GameplayDebuggerCategory.cpp
deleted file mode 100644
diff --git a/Engine/Source/Developer/GameplayDebugger/Private/GameplayDebuggerCategoryReplicator.cpp b/Engine/Source/Developer/GameplayDebugger/Private/GameplayDebuggerCategoryReplicator.cpp
deleted file mode 100644
diff --git a/Engine/Source/Developer/GameplayDebugger/Private/GameplayDebuggerCompat.h b/Engine/Source/Developer/GameplayDebugger/Private/GameplayDebuggerCompat.h
deleted file mode 100644
diff --git a/Engine/Source/Developer/GameplayDebugger/Private/GameplayDebuggerConfig.cpp b/Engine/Source/Developer/GameplayDebugger/Private/GameplayDebuggerConfig.cpp
deleted file mode 100644
diff --git a/Engine/Source/Developer/GameplayDebugger/Private/GameplayDebuggerExtension.cpp b/Engine/Source/Developer/GameplayDebugger/Private/GameplayDebuggerExtension.cpp
deleted file mode 100644
diff --git a/Engine/Source/Developer/GameplayDebugger/Private/GameplayDebuggerExtension_HUD.cpp b/Engine/Source/Developer/GameplayDebugger/Private/GameplayDebuggerExtension_HUD.cpp
deleted file mode 100644
diff --git a/Engine/Source/Developer/GameplayDebugger/Private/GameplayDebuggerExtension_HUD.h b/Engine/Source/Developer/GameplayDebugger/Private/GameplayDebuggerExtension_HUD.h
deleted file mode 100644
diff --git a/Engine/Source/Developer/GameplayDebugger/Private/GameplayDebuggerExtension_Spectator.cpp b/Engine/Source/Developer/GameplayDebugger/Private/GameplayDebuggerExtension_Spectator.cpp
deleted file mode 100644
diff --git a/Engine/Source/Developer/GameplayDebugger/Private/GameplayDebuggerExtension_Spectator.h b/Engine/Source/Developer/GameplayDebugger/Private/GameplayDebuggerExtension_Spectator.h
deleted file mode 100644
diff --git a/Engine/Source/Developer/GameplayDebugger/Private/GameplayDebuggerLocalController.cpp b/Engine/Source/Developer/GameplayDebugger/Private/GameplayDebuggerLocalController.cpp
deleted file mode 100644
diff --git a/Engine/Source/Developer/GameplayDebugger/Private/GameplayDebuggerModule.cpp b/Engine/Source/Developer/GameplayDebugger/Private/GameplayDebuggerModule.cpp
deleted file mode 100644
diff --git a/Engine/Source/Developer/GameplayDebugger/Private/GameplayDebuggerPlayerManager.cpp b/Engine/Source/Developer/GameplayDebugger/Private/GameplayDebuggerPlayerManager.cpp
deleted file mode 100644
diff --git a/Engine/Source/Developer/GameplayDebugger/Private/GameplayDebuggerRenderingComponent.cpp b/Engine/Source/Developer/GameplayDebugger/Private/GameplayDebuggerRenderingComponent.cpp
deleted file mode 100644
diff --git a/Engine/Source/Developer/GameplayDebugger/Private/GameplayDebuggerTypes.cpp b/Engine/Source/Developer/GameplayDebugger/Private/GameplayDebuggerTypes.cpp
deleted file mode 100644
diff --git a/Engine/Source/Developer/GameplayDebugger/Public/GameplayDebugger.h b/Engine/Source/Developer/GameplayDebugger/Public/GameplayDebugger.h
deleted file mode 100644
diff --git a/Engine/Source/Developer/GameplayDebugger/Public/GameplayDebuggerAddonBase.h b/Engine/Source/Developer/GameplayDebugger/Public/GameplayDebuggerAddonBase.h
deleted file mode 100644
diff --git a/Engine/Source/Developer/GameplayDebugger/Public/GameplayDebuggerCategory.h b/Engine/Source/Developer/GameplayDebugger/Public/GameplayDebuggerCategory.h
deleted file mode 100644
diff --git a/Engine/Source/Developer/GameplayDebugger/Public/GameplayDebuggerCategoryReplicator.h b/Engine/Source/Developer/GameplayDebugger/Public/GameplayDebuggerCategoryReplicator.h
deleted file mode 100644
diff --git a/Engine/Source/Developer/GameplayDebugger/Public/GameplayDebuggerConfig.h b/Engine/Source/Developer/GameplayDebugger/Public/GameplayDebuggerConfig.h
deleted file mode 100644
diff --git a/Engine/Source/Developer/GameplayDebugger/Public/GameplayDebuggerExtension.h b/Engine/Source/Developer/GameplayDebugger/Public/GameplayDebuggerExtension.h
deleted file mode 100644
diff --git a/Engine/Source/Developer/GameplayDebugger/Public/GameplayDebuggerLocalController.h b/Engine/Source/Developer/GameplayDebugger/Public/GameplayDebuggerLocalController.h
deleted file mode 100644
diff --git a/Engine/Source/Developer/GameplayDebugger/Public/GameplayDebuggerPlayerManager.h b/Engine/Source/Developer/GameplayDebugger/Public/GameplayDebuggerPlayerManager.h
deleted file mode 100644
diff --git a/Engine/Source/Developer/GameplayDebugger/Public/GameplayDebuggerRenderingComponent.h b/Engine/Source/Developer/GameplayDebugger/Public/GameplayDebuggerRenderingComponent.h
deleted file mode 100644
diff --git a/Engine/Source/Developer/GameplayDebugger/Public/GameplayDebuggerTypes.h b/Engine/Source/Developer/GameplayDebugger/Public/GameplayDebuggerTypes.h
deleted file mode 100644
diff --git a/Engine/Source/Editor/GameplayDebugger/GameplayDebuggerEditor.Build.cs b/Engine/Source/Editor/GameplayDebugger/GameplayDebuggerEditor.Build.cs
new file mode 100644
--- /dev/null
+++ b/Engine/Source/Editor/GameplayDebugger/GameplayDebuggerEditor.Build.cs
@@ -0,0 +1,54 @@
+// Copyright Epic Games, Inc. All Rights Reserved.
+
+namespace UnrealBuildTool.Rules
+{
+	public class GameplayDebuggerEditor : ModuleRules
+	{
+		public GameplayDebuggerEditor(ReadOnlyTargetRules Target) : base(Target)
+		{
+			PublicDependencyModuleNames.AddRange(
+				new string[] {
+					"Core",
+					"CoreUObject",
+					"Engine",
+				});
+
+			PrivateDependencyModuleNames.AddRange(
+				new string[] {
+					"RenderCore",
+					"InputCore",
+					"SlateCore",
+					"Slate",
+					"DeveloperSettings",
+				});
+			
+			if (Target.Configuration != UnrealTargetConfiguration.Shipping)
+			{
+				PrivateDependencyModuleNames.Add("DrawPrimitiveDebugger");
+			}
+/*
+			PrivateIncludePaths.AddRange(
+				new string[] {
+					"Developer/Settings/Public",
+				});
+*/
+			if (Target.bBuildEditor)
+			{
+				PrivateDependencyModuleNames.AddRange(
+					new string[] {
+						
+						"EditorFramework",
+						"UnrealEd",
+						"LevelEditor",
+						"PropertyEditor",
+						"GameplayDebugger"
+					});
+			}
+
+			if (Target.bUseGameplayDebugger)
+			{
+				PrecompileForTargets = PrecompileTargetsType.Any;
+			}
+		}
+	}
+}
diff --git a/Engine/Source/Editor/GameplayDebugger/Private/GameplayDebuggerCategoryConfigCustomization.cpp b/Engine/Source/Editor/GameplayDebugger/Private/GameplayDebuggerCategoryConfigCustomization.cpp
new file mode 100644
--- /dev/null
+++ b/Engine/Source/Editor/GameplayDebugger/Private/GameplayDebuggerCategoryConfigCustomization.cpp
@@ -0,0 +1,92 @@
+// Copyright Epic Games, Inc. All Rights Reserved.
+
+#include "GameplayDebuggerCategoryConfigCustomization.h"
+
+#include "UObject/UnrealType.h"
+#include "Widgets/DeclarativeSyntaxSupport.h"
+#include "Application/SlateWindowHelper.h"
+#include "Widgets/Text/STextBlock.h"
+#include "DetailWidgetRow.h"
+#include "IDetailChildrenBuilder.h"
+#include "DetailLayoutBuilder.h"
+#include "GameplayDebuggerConfig.h"
+
+#define LOCTEXT_NAMESPACE "GameplayDebuggerConfig"
+
+void FGameplayDebuggerCategoryConfigCustomization::CustomizeHeader(TSharedRef<class IPropertyHandle> StructPropertyHandle, class FDetailWidgetRow& HeaderRow, IPropertyTypeCustomizationUtils& StructCustomizationUtils)
+{
+	CategoryNameProp = StructPropertyHandle->GetChildHandle(GET_MEMBER_NAME_CHECKED(FGameplayDebuggerCategoryConfig, CategoryName));
+	SlotIdxProp = StructPropertyHandle->GetChildHandle(GET_MEMBER_NAME_CHECKED(FGameplayDebuggerCategoryConfig, SlotIdx));
+	ActiveInGameProp = StructPropertyHandle->GetChildHandle(GET_MEMBER_NAME_CHECKED(FGameplayDebuggerCategoryConfig, ActiveInGame));
+	ActiveInSimulateProp = StructPropertyHandle->GetChildHandle(GET_MEMBER_NAME_CHECKED(FGameplayDebuggerCategoryConfig, ActiveInSimulate));
+	
+	FSimpleDelegate RefreshHeaderDelegate = FSimpleDelegate::CreateSP(this, &FGameplayDebuggerCategoryConfigCustomization::OnChildValueChanged);
+	StructPropertyHandle->SetOnChildPropertyValueChanged(RefreshHeaderDelegate);
+	OnChildValueChanged();
+
+	HeaderRow.NameContent()
+	[
+		StructPropertyHandle->CreatePropertyNameWidget()
+	]
+	.ValueContent().VAlign(VAlign_Center).MinDesiredWidth(300.0f)
+	[
+		SNew(STextBlock)
+		.Text(this, &FGameplayDebuggerCategoryConfigCustomization::GetHeaderDesc)
+		.Font(IDetailLayoutBuilder::GetDetailFont())
+	];
+};
+
+void FGameplayDebuggerCategoryConfigCustomization::CustomizeChildren(TSharedRef<class IPropertyHandle> StructPropertyHandle, class IDetailChildrenBuilder& StructBuilder, IPropertyTypeCustomizationUtils& StructCustomizationUtils)
+{
+	uint32 NumChildProps = 0;
+	StructPropertyHandle->GetNumChildren(NumChildProps);
+
+	for (uint32 Idx = 0; Idx < NumChildProps; Idx++)
+	{
+		TSharedPtr<IPropertyHandle> PropHandle = StructPropertyHandle->GetChildHandle(Idx);
+		if (PropHandle->GetProperty() && PropHandle->GetProperty()->GetFName() == GET_MEMBER_NAME_CHECKED(FGameplayDebuggerCategoryConfig, CategoryName))
+		{
+			continue;
+		}
+
+		StructBuilder.AddProperty(PropHandle.ToSharedRef());
+	}
+};
+
+void FGameplayDebuggerCategoryConfigCustomization::OnChildValueChanged()
+{
+	FString CategoryNameValue;
+	if (CategoryNameProp.IsValid())
+	{
+		CategoryNameProp->GetValue(CategoryNameValue);
+	}
+
+	int32 SlotIdxValue = -1;
+	if (SlotIdxProp.IsValid())
+	{
+		SlotIdxProp->GetValue(SlotIdxValue);
+	}
+
+	uint8 ActiveInGameValue = 0;
+	if (ActiveInGameProp.IsValid())
+	{
+		ActiveInGameProp->GetValue(ActiveInGameValue);
+	}
+
+	uint8 ActiveInSimulateValue = 0;
+	if (ActiveInSimulateProp.IsValid())
+	{
+		ActiveInSimulateProp->GetValue(ActiveInSimulateValue);
+	}
+
+	CachedHeader = FText::FromString(FString::Printf(TEXT("[%s]:%s%s%s"),
+		SlotIdxValue < 0 ? TEXT("-") : *TTypeToString<int32>::ToString(SlotIdxValue),
+		CategoryNameValue.Len() ? *CategoryNameValue : TEXT("??"),
+		(ActiveInGameValue == (uint8)EGameplayDebuggerOverrideMode::UseDefault) ? TEXT("") :
+			(ActiveInGameValue == (uint8)EGameplayDebuggerOverrideMode::Enable) ? TEXT(" game:ON") : TEXT(" game:OFF"),
+		(ActiveInSimulateValue == (uint8)EGameplayDebuggerOverrideMode::UseDefault) ? TEXT("") :
+			(ActiveInSimulateValue == (uint8)EGameplayDebuggerOverrideMode::Enable) ? TEXT(" simulate:ON") : TEXT(" simulate:OFF")
+		));
+}
+
+#undef LOCTEXT_NAMESPACE
diff --git a/Engine/Source/Editor/GameplayDebugger/Private/GameplayDebuggerEdMode.cpp b/Engine/Source/Editor/GameplayDebugger/Private/GameplayDebuggerEdMode.cpp
new file mode 100644
--- /dev/null
+++ b/Engine/Source/Editor/GameplayDebugger/Private/GameplayDebuggerEdMode.cpp
@@ -0,0 +1,140 @@
+// Copyright Epic Games, Inc. All Rights Reserved.
+
+#include "GameplayDebuggerEdMode.h"
+
+#include "CoreGlobals.h"
+#include "GameplayDebuggerToolkit.h"
+#include "Toolkits/ToolkitManager.h"
+#include "GameplayDebuggerPlayerManager.h"
+#include "GameplayDebuggerLocalController.h"
+
+#include "EditorModeManager.h"
+#include "Components/InputComponent.h"
+
+const FName FGameplayDebuggerEdMode::EM_GameplayDebugger = TEXT("EM_GameplayDebugger");
+
+void FGameplayDebuggerEdMode::Enter()
+{
+	FEdMode::Enter();
+	
+	if (!Toolkit.IsValid())
+	{
+		Toolkit = MakeShareable(new FGameplayDebuggerToolkit(this));
+		Toolkit->Init(Owner->GetToolkitHost());
+	}
+
+	bPrevGAreScreenMessagesEnabled = GAreScreenMessagesEnabled;
+	GAreScreenMessagesEnabled = false;
+}
+
+void FGameplayDebuggerEdMode::Exit()
+{
+	if (Toolkit.IsValid())
+	{
+		FToolkitManager::Get().CloseToolkit(Toolkit.ToSharedRef());
+		Toolkit.Reset();
+	}
+
+	if (FocusedViewport)
+	{
+		EnableViewportClientFlags(FocusedViewport, false);
+		FocusedViewport = nullptr;
+	}
+
+	FEdMode::Exit();
+	GAreScreenMessagesEnabled = bPrevGAreScreenMessagesEnabled;
+}
+
+void FGameplayDebuggerEdMode::EnableViewportClientFlags(FEditorViewportClient* ViewportClient, bool bEnable)
+{
+	ViewportClient->bUseNumpadCameraControl = false;
+}
+
+bool FGameplayDebuggerEdMode::ReceivedFocus(FEditorViewportClient* ViewportClient, FViewport* Viewport)
+{
+	EnableViewportClientFlags(ViewportClient, true);
+	FocusedViewport = ViewportClient;
+	return false;
+}
+
+bool FGameplayDebuggerEdMode::LostFocus(FEditorViewportClient* ViewportClient, FViewport* Viewport)
+{
+	EnableViewportClientFlags(ViewportClient, false);
+	FocusedViewport = nullptr;
+	return false;
+}
+
+bool FGameplayDebuggerEdMode::InputKey(FEditorViewportClient* InViewportClient, FViewport* InViewport, FKey InKey, EInputEvent InEvent)
+{
+	UWorld* MyWorld = GetWorld();
+	APlayerController* LocalPC = GEngine->GetFirstLocalPlayerController(MyWorld);
+
+	if (LocalPC && InViewportClient)
+	{
+		// process raw input for debugger's input component manually
+		// can't use LocalPC->InputKey here, since it will trigger for every bound chord, not only gameplay debugger ones
+		// and will not work at all when simulation is paused
+
+		AGameplayDebuggerPlayerManager& PlayerManager = AGameplayDebuggerPlayerManager::GetCurrent(MyWorld);
+		const FGameplayDebuggerPlayerData* DataPtr = PlayerManager.GetPlayerData(*LocalPC);
+
+		if (DataPtr && DataPtr->InputComponent && DataPtr->Controller && DataPtr->Controller->IsKeyBound(InKey.GetFName()))
+		{
+			const FInputChord ActiveChord(InKey, InViewportClient->IsShiftPressed(), InViewportClient->IsCtrlPressed(), InViewportClient->IsAltPressed(), InViewportClient->IsCmdPressed());
+
+			// go over all bound actions
+			const int32 NumBindings = DataPtr->InputComponent->KeyBindings.Num();
+			for (int32 Idx = 0; Idx < NumBindings; Idx++)
+			{
+				const FInputKeyBinding& KeyBinding = DataPtr->InputComponent->KeyBindings[Idx];
+				if (KeyBinding.KeyEvent == InEvent && KeyBinding.Chord == ActiveChord && KeyBinding.KeyDelegate.IsBound())
+				{
+					KeyBinding.KeyDelegate.Execute(InKey);
+				}
+			}
+
+			return true;			
+		}
+	}
+
+	return false;
+}
+
+void FGameplayDebuggerEdMode::Tick(FEditorViewportClient* ViewportClient, float DeltaTime)
+{
+	if (ViewportClient == nullptr || !ViewportClient->EngineShowFlags.DebugAI)
+	{
+		Owner->DeactivateMode(FGameplayDebuggerEdMode::EM_GameplayDebugger);
+	}
+}
+
+void FGameplayDebuggerEdMode::SafeOpenMode()
+{
+	// The global mode manager should not be created or accessed in a commandlet environment
+	if (IsRunningCommandlet())
+	{
+		return;
+	}
+
+	// By calling this get, it make sure that the singleton is created and not during the SafeCloseMode whitch can be called during garbage collect.
+	GLevelEditorModeTools();
+}
+
+void FGameplayDebuggerEdMode::SafeCloseMode()
+{
+	// The global mode manager should not be created or accessed in a commandlet environment
+	if (IsRunningCommandlet())
+	{
+		return;
+	}
+
+	// this may be called on closing editor during PIE (~viewport -> teardown PIE -> debugger's cleanup on game end)
+	//
+	// DeactivateMode tries to bring up default mode, but toolkit is already destroyed by that time
+	// and editor crashes on check in GLevelEditorModeTools().GetToolkitHost() inside default mode's code
+	// note: no need to do it if IsEngineExitRequested since the GLevelEditorModeTools is no longer available 
+	if (IsEngineExitRequested() == false && GLevelEditorModeTools().HasToolkitHost())
+	{
+		GLevelEditorModeTools().DeactivateMode(FGameplayDebuggerEdMode::EM_GameplayDebugger);
+	}
+}
diff --git a/Engine/Source/Editor/GameplayDebugger/Private/GameplayDebuggerEditorModule.cpp b/Engine/Source/Editor/GameplayDebugger/Private/GameplayDebuggerEditorModule.cpp
new file mode 100644
--- /dev/null
+++ b/Engine/Source/Editor/GameplayDebugger/Private/GameplayDebuggerEditorModule.cpp
@@ -0,0 +1,59 @@
+// Copyright Epic Games, Inc. All Rights Reserved.
+
+#include "GameplayDebuggerEditorModule.h"
+
+#include "CoreMinimal.h"
+#include "EditorModeManager.h"
+#include "EditorModeRegistry.h"
+#include "Engine/World.h"
+#include "GameplayDebugger.h"
+#include "GameplayDebuggerCategoryConfigCustomization.h"
+#include "GameplayDebuggerEdMode.h"
+#include "GameplayDebuggerExtensionConfigCustomization.h"
+#include "GameplayDebuggerInputConfigCustomization.h"
+#include "GameplayDebuggerModule.h"
+#include "PropertyEditorModule.h"
+
+IMPLEMENT_MODULE(FGameplayDebuggerEditorModule, GameplayDebuggerEditor)
+
+void FGameplayDebuggerEditorModule::StartupModule()
+{
+	FGameplayDebuggerModule::OnLocalControllerInitialized.AddRaw(this, &FGameplayDebuggerEditorModule::OnLocalControllerInitialized);
+	FGameplayDebuggerModule::OnLocalControllerUninitialized.AddRaw(this, &FGameplayDebuggerEditorModule::OnLocalControllerUninitialized);
+	FGameplayDebuggerModule::OnDebuggerEdModeActivation.AddRaw(this, &FGameplayDebuggerEditorModule::OnDebuggerEdModeActivation);
+
+	FPropertyEditorModule& PropertyEditorModule = FModuleManager::GetModuleChecked<FPropertyEditorModule>("PropertyEditor");
+	PropertyEditorModule.RegisterCustomPropertyTypeLayout("GameplayDebuggerCategoryConfig", FOnGetPropertyTypeCustomizationInstance::CreateStatic(&FGameplayDebuggerCategoryConfigCustomization::MakeInstance));
+	PropertyEditorModule.RegisterCustomPropertyTypeLayout("GameplayDebuggerExtensionConfig", FOnGetPropertyTypeCustomizationInstance::CreateStatic(&FGameplayDebuggerExtensionConfigCustomization::MakeInstance));
+	PropertyEditorModule.RegisterCustomPropertyTypeLayout("GameplayDebuggerInputConfig", FOnGetPropertyTypeCustomizationInstance::CreateStatic(&FGameplayDebuggerInputConfigCustomization::MakeInstance));
+
+	FEditorModeRegistry::Get().RegisterMode<FGameplayDebuggerEdMode>(FGameplayDebuggerEdMode::EM_GameplayDebugger);
+}
+
+void FGameplayDebuggerEditorModule::ShutdownModule()
+{	
+	FPropertyEditorModule* PropertyEditorModule = FModuleManager::GetModulePtr<FPropertyEditorModule>("PropertyEditor");
+	if (PropertyEditorModule)
+	{
+		PropertyEditorModule->UnregisterCustomPropertyTypeLayout("GameplayDebuggerCategoryConfig");
+		PropertyEditorModule->UnregisterCustomPropertyTypeLayout("GameplayDebuggerExtensionConfig");
+		PropertyEditorModule->UnregisterCustomPropertyTypeLayout("GameplayDebuggerInputConfig");
+	}
+
+	FEditorModeRegistry::Get().UnregisterMode(FGameplayDebuggerEdMode::EM_GameplayDebugger);
+}
+
+void FGameplayDebuggerEditorModule::OnDebuggerEdModeActivation()
+{
+	GLevelEditorModeTools().ActivateMode(FGameplayDebuggerEdMode::EM_GameplayDebugger);
+}
+
+void FGameplayDebuggerEditorModule::OnLocalControllerInitialized()
+{
+	FGameplayDebuggerEdMode::SafeOpenMode();
+}
+
+void FGameplayDebuggerEditorModule::OnLocalControllerUninitialized()
+{
+	FGameplayDebuggerEdMode::SafeCloseMode();
+}
diff --git a/Engine/Source/Editor/GameplayDebugger/Private/GameplayDebuggerExtensionConfigCustomization.cpp b/Engine/Source/Editor/GameplayDebugger/Private/GameplayDebuggerExtensionConfigCustomization.cpp
new file mode 100644
--- /dev/null
+++ b/Engine/Source/Editor/GameplayDebugger/Private/GameplayDebuggerExtensionConfigCustomization.cpp
@@ -0,0 +1,75 @@
+// Copyright Epic Games, Inc. All Rights Reserved.
+
+#include "GameplayDebuggerExtensionConfigCustomization.h"
+
+#include "UObject/UnrealType.h"
+#include "Widgets/DeclarativeSyntaxSupport.h"
+#include "Application/SlateWindowHelper.h"
+#include "Widgets/Text/STextBlock.h"
+#include "DetailWidgetRow.h"
+#include "IDetailChildrenBuilder.h"
+#include "DetailLayoutBuilder.h"
+#include "GameplayDebuggerConfig.h"
+
+#define LOCTEXT_NAMESPACE "GameplayDebuggerConfig"
+
+void FGameplayDebuggerExtensionConfigCustomization::CustomizeHeader(TSharedRef<class IPropertyHandle> StructPropertyHandle, class FDetailWidgetRow& HeaderRow, IPropertyTypeCustomizationUtils& StructCustomizationUtils)
+{
+	ExtensionNameProp = StructPropertyHandle->GetChildHandle(GET_MEMBER_NAME_CHECKED(FGameplayDebuggerExtensionConfig, ExtensionName));
+	UseExtensionProp = StructPropertyHandle->GetChildHandle(GET_MEMBER_NAME_CHECKED(FGameplayDebuggerExtensionConfig, UseExtension));
+	
+	FSimpleDelegate RefreshHeaderDelegate = FSimpleDelegate::CreateSP(this, &FGameplayDebuggerExtensionConfigCustomization::OnChildValueChanged);
+	StructPropertyHandle->SetOnChildPropertyValueChanged(RefreshHeaderDelegate);
+	OnChildValueChanged();
+
+	HeaderRow.NameContent()
+	[
+		StructPropertyHandle->CreatePropertyNameWidget()
+	]
+	.ValueContent().VAlign(VAlign_Center).MinDesiredWidth(300.0f)
+	[
+		SNew(STextBlock)
+		.Text(this, &FGameplayDebuggerExtensionConfigCustomization::GetHeaderDesc)
+		.Font(IDetailLayoutBuilder::GetDetailFont())
+	];
+};
+
+void FGameplayDebuggerExtensionConfigCustomization::CustomizeChildren(TSharedRef<class IPropertyHandle> StructPropertyHandle, class IDetailChildrenBuilder& StructBuilder, IPropertyTypeCustomizationUtils& StructCustomizationUtils)
+{
+	uint32 NumChildProps = 0;
+	StructPropertyHandle->GetNumChildren(NumChildProps);
+
+	for (uint32 Idx = 0; Idx < NumChildProps; Idx++)
+	{
+		TSharedPtr<IPropertyHandle> PropHandle = StructPropertyHandle->GetChildHandle(Idx);
+		if (PropHandle->GetProperty() && PropHandle->GetProperty()->GetFName() == GET_MEMBER_NAME_CHECKED(FGameplayDebuggerExtensionConfig, ExtensionName))
+		{
+			continue;
+		}
+
+		StructBuilder.AddProperty(PropHandle.ToSharedRef());
+	}
+};
+
+void FGameplayDebuggerExtensionConfigCustomization::OnChildValueChanged()
+{
+	FString ExtensionNameValue;
+	if (ExtensionNameProp.IsValid())
+	{
+		ExtensionNameProp->GetValue(ExtensionNameValue);
+	}
+
+	uint8 UseExtensionValue = 0;
+	if (UseExtensionProp.IsValid())
+	{
+		UseExtensionProp->GetValue(UseExtensionValue);
+	}
+
+	CachedHeader = FText::FromString(FString::Printf(TEXT("%s %s"),
+		ExtensionNameValue.Len() ? *ExtensionNameValue : TEXT("??"),
+		(UseExtensionValue == (uint8)EGameplayDebuggerOverrideMode::UseDefault) ? TEXT("") :
+			(UseExtensionValue == (uint8)EGameplayDebuggerOverrideMode::Enable) ? TEXT("is enabled") : TEXT("is disabled")
+		));
+}
+
+#undef LOCTEXT_NAMESPACE
diff --git a/Engine/Source/Editor/GameplayDebugger/Private/GameplayDebuggerInputConfigCustomization.cpp b/Engine/Source/Editor/GameplayDebugger/Private/GameplayDebuggerInputConfigCustomization.cpp
new file mode 100644
--- /dev/null
+++ b/Engine/Source/Editor/GameplayDebugger/Private/GameplayDebuggerInputConfigCustomization.cpp
@@ -0,0 +1,109 @@
+// Copyright Epic Games, Inc. All Rights Reserved.
+
+#include "GameplayDebuggerInputConfigCustomization.h"
+#include "Widgets/DeclarativeSyntaxSupport.h"
+#include "Engine/GameViewportClient.h"
+#include "GameplayDebuggerTypes.h"
+
+#include "Widgets/Text/STextBlock.h"
+#include "DetailWidgetRow.h"
+#include "IDetailChildrenBuilder.h"
+#include "DetailLayoutBuilder.h"
+#include "GameplayDebuggerConfig.h"
+
+#define LOCTEXT_NAMESPACE "GameplayDebuggerConfig"
+
+void FGameplayDebuggerInputConfigCustomization::CustomizeHeader(TSharedRef<class IPropertyHandle> StructPropertyHandle, class FDetailWidgetRow& HeaderRow, IPropertyTypeCustomizationUtils& StructCustomizationUtils)
+{
+	ConfigNameProp = StructPropertyHandle->GetChildHandle(GET_MEMBER_NAME_CHECKED(FGameplayDebuggerInputConfig, ConfigName));
+	KeyProp = StructPropertyHandle->GetChildHandle(GET_MEMBER_NAME_CHECKED(FGameplayDebuggerInputConfig, Key));
+	ModShiftProp = StructPropertyHandle->GetChildHandle(GET_MEMBER_NAME_CHECKED(FGameplayDebuggerInputConfig, bModShift));
+	ModCtrlProp = StructPropertyHandle->GetChildHandle(GET_MEMBER_NAME_CHECKED(FGameplayDebuggerInputConfig, bModCtrl));
+	ModAltProp = StructPropertyHandle->GetChildHandle(GET_MEMBER_NAME_CHECKED(FGameplayDebuggerInputConfig, bModAlt));
+	ModCmdProp = StructPropertyHandle->GetChildHandle(GET_MEMBER_NAME_CHECKED(FGameplayDebuggerInputConfig, bModCmd));
+
+	FSimpleDelegate RefreshHeaderDelegate = FSimpleDelegate::CreateSP(this, &FGameplayDebuggerInputConfigCustomization::OnChildValueChanged);
+	StructPropertyHandle->SetOnChildPropertyValueChanged(RefreshHeaderDelegate);
+	OnChildValueChanged();
+
+	HeaderRow.NameContent()
+	[
+		StructPropertyHandle->CreatePropertyNameWidget()
+	]
+	.ValueContent().VAlign(VAlign_Center).MinDesiredWidth(300.0f)
+	[
+		SNew(STextBlock)
+		.Text(this, &FGameplayDebuggerInputConfigCustomization::GetHeaderDesc)
+		.Font(IDetailLayoutBuilder::GetDetailFont())
+	];
+};
+
+void FGameplayDebuggerInputConfigCustomization::CustomizeChildren(TSharedRef<class IPropertyHandle> StructPropertyHandle, class IDetailChildrenBuilder& StructBuilder, IPropertyTypeCustomizationUtils& StructCustomizationUtils)
+{
+	uint32 NumChildProps = 0;
+	StructPropertyHandle->GetNumChildren(NumChildProps);
+
+	for (uint32 Idx = 0; Idx < NumChildProps; Idx++)
+	{
+		TSharedPtr<IPropertyHandle> PropHandle = StructPropertyHandle->GetChildHandle(Idx);
+		if (PropHandle->GetProperty() && PropHandle->GetProperty()->GetFName() == GET_MEMBER_NAME_CHECKED(FGameplayDebuggerInputConfig, ConfigName))
+		{
+			continue;
+		}
+
+		StructBuilder.AddProperty(PropHandle.ToSharedRef());
+	}
+};
+
+void FGameplayDebuggerInputConfigCustomization::OnChildValueChanged()
+{
+	FString ConfigNameValue;
+	if (ConfigNameProp.IsValid())
+	{
+		ConfigNameProp->GetValue(ConfigNameValue);
+	}
+
+	FString KeyValue;
+	if (KeyProp.IsValid())
+	{
+		KeyProp->GetPerObjectValue(0, KeyValue);
+	}
+
+	bool bModShift = false;
+	if (ModShiftProp.IsValid())
+	{
+		ModShiftProp->GetValue(bModShift);
+	}
+
+	bool bModCtrl = false;
+	if (ModCtrlProp.IsValid())
+	{
+		ModCtrlProp->GetValue(bModCtrl);
+	}
+	
+	bool bModAlt = false;
+	if (ModAltProp.IsValid())
+	{
+		ModAltProp->GetValue(bModAlt);
+	}
+	
+	bool bModCmd = false;
+	if (ModCmdProp.IsValid())
+	{
+		ModCmdProp->GetValue(bModCmd);
+	}
+
+	FGameplayDebuggerInputHandler DummyHandler;
+	DummyHandler.KeyName = *KeyValue;
+	DummyHandler.Modifier.bShift = bModShift;
+	DummyHandler.Modifier.bCtrl = bModCtrl;
+	DummyHandler.Modifier.bAlt = bModAlt;
+	DummyHandler.Modifier.bCmd = bModCmd;
+
+	CachedHeader = FText::FromString(FString::Printf(TEXT("%s: %s"),
+		ConfigNameValue.Len() ? *ConfigNameValue : TEXT("??"),
+		*DummyHandler.ToString()
+		));
+}
+
+#undef LOCTEXT_NAMESPACE
diff --git a/Engine/Source/Editor/GameplayDebugger/Private/GameplayDebuggerToolkit.cpp b/Engine/Source/Editor/GameplayDebugger/Private/GameplayDebuggerToolkit.cpp
new file mode 100644
--- /dev/null
+++ b/Engine/Source/Editor/GameplayDebugger/Private/GameplayDebuggerToolkit.cpp
@@ -0,0 +1,106 @@
+// Copyright Epic Games, Inc. All Rights Reserved.
+
+#include "GameplayDebuggerToolkit.h"
+#include "Modules/ModuleManager.h"
+
+#include "LevelEditor.h"
+#include "IAssetViewport.h"
+#include "SLevelViewport.h"
+#include "LevelEditorViewport.h"
+#include "Widgets/Layout/SBorder.h"
+#include "Widgets/Input/SButton.h"
+#include "Widgets/Text/STextBlock.h"
+#include "Styling/AppStyle.h"
+
+#define LOCTEXT_NAMESPACE "GameplayDebuggerConfig"
+
+FGameplayDebuggerToolkit::FGameplayDebuggerToolkit(class FEdMode* InOwningMode)
+{
+	DebuggerEdMode = InOwningMode;
+}
+
+FText FGameplayDebuggerToolkit::GetBaseToolkitName() const
+{
+	return LOCTEXT("EdToolTitle", "Gameplay Debugger");
+}
+
+FName FGameplayDebuggerToolkit::GetToolkitFName() const
+{
+	return FName("GameplayDebuggerToolkit");
+}
+
+void FGameplayDebuggerToolkit::Init(const TSharedPtr<class IToolkitHost>& InitToolkitHost)
+{
+	MyWidget =
+		SNew(SBorder)
+		.BorderImage(FAppStyle::GetBrush("ToolPanel.GroupBorder"))
+		.Content()
+		[
+			SNew(SVerticalBox)
+			+ SVerticalBox::Slot()
+				.AutoHeight()
+				.HAlign(HAlign_Left)
+				.VAlign(VAlign_Center)
+				[
+					SNew(STextBlock)
+					.Text(LOCTEXT("EdToolStatus", "Gameplay debugger is active."))
+				]
+			+ SVerticalBox::Slot()
+				.AutoHeight()
+				.HAlign(HAlign_Left)
+				.VAlign(VAlign_Center)
+				.Padding(0, 30.0f)
+				[
+					SNew(STextBlock)
+					.Visibility(this, &FGameplayDebuggerToolkit::GetScreenMessageWarningVisibility)
+					.Text(LOCTEXT("EdToolMessageWarning", "Warning! On screen messages are suppressed!\nUse EnableAllScreenMessages to restore them."))
+				]
+			+ SVerticalBox::Slot()
+				.AutoHeight()
+				.HAlign(HAlign_Left)
+				.VAlign(VAlign_Center)
+				[
+					SNew(STextBlock)
+					.Text(LOCTEXT("EdToolDisableHint", "Clear DebugAI show flag to disable tool."))
+				]
+			+ SVerticalBox::Slot()
+				.AutoHeight()
+				.HAlign(HAlign_Left)
+				.VAlign(VAlign_Center)
+				.Padding(5.0f, 5.0f)
+				[
+					SNew(SButton)
+					.Text(LOCTEXT("EdToolDisableButton", "Disable tool"))
+					.OnClicked(this, &FGameplayDebuggerToolkit::OnClickedDisableTool)
+				]
+		];
+
+	FModeToolkit::Init(InitToolkitHost);
+}
+
+EVisibility FGameplayDebuggerToolkit::GetScreenMessageWarningVisibility() const
+{
+	return GAreScreenMessagesEnabled ? EVisibility::Collapsed : EVisibility::Visible;
+}
+
+FReply FGameplayDebuggerToolkit::OnClickedDisableTool()
+{
+	FLevelEditorModule& LevelEditorModule = FModuleManager::GetModuleChecked<FLevelEditorModule>("LevelEditor");
+	TSharedPtr<ILevelEditor> LevelEditor = LevelEditorModule.GetFirstLevelEditor();
+	if (LevelEditor.IsValid())
+	{
+		TArray<TSharedPtr<SLevelViewport>> Viewports = LevelEditor->GetViewports();
+		for (const TSharedPtr<SLevelViewport>& ViewportWindow : Viewports)
+		{
+			if (ViewportWindow.IsValid())
+			{
+				FEditorViewportClient& Viewport = ViewportWindow->GetAssetViewportClient();
+				Viewport.EngineShowFlags.SetDebugAI(false);
+			}
+		}
+	}
+	
+	return FReply::Handled();
+}
+
+#undef LOCTEXT_NAMESPACE
diff --git a/Engine/Source/Editor/GameplayDebugger/Public/GameplayDebuggerCategoryConfigCustomization.h b/Engine/Source/Editor/GameplayDebugger/Public/GameplayDebuggerCategoryConfigCustomization.h
new file mode 100644
--- /dev/null
+++ b/Engine/Source/Editor/GameplayDebugger/Public/GameplayDebuggerCategoryConfigCustomization.h
@@ -0,0 +1,32 @@
+// Copyright Epic Games, Inc. All Rights Reserved.
+
+#pragma once
+
+#include "CoreMinimal.h"
+
+#include "IPropertyTypeCustomization.h"
+
+class IPropertyHandle;
+
+class FGameplayDebuggerCategoryConfigCustomization : public IPropertyTypeCustomization
+{
+public:
+	static TSharedRef<IPropertyTypeCustomization> MakeInstance()
+	{
+		return MakeShareable(new FGameplayDebuggerCategoryConfigCustomization);
+	}
+
+	virtual void CustomizeHeader(TSharedRef<class IPropertyHandle> StructPropertyHandle, class FDetailWidgetRow& HeaderRow, IPropertyTypeCustomizationUtils& StructCustomizationUtils) override;
+	virtual void CustomizeChildren(TSharedRef<class IPropertyHandle> StructPropertyHandle, class IDetailChildrenBuilder& StructBuilder, IPropertyTypeCustomizationUtils& StructCustomizationUtils) override;
+
+protected:
+	TSharedPtr<IPropertyHandle> CategoryNameProp;
+	TSharedPtr<IPropertyHandle> SlotIdxProp;
+	TSharedPtr<IPropertyHandle> ActiveInGameProp;
+	TSharedPtr<IPropertyHandle> ActiveInSimulateProp;
+
+	FText CachedHeader;
+	FText GetHeaderDesc() const { return CachedHeader; }
+
+	void OnChildValueChanged();
+};
diff --git a/Engine/Source/Editor/GameplayDebugger/Public/GameplayDebuggerEdMode.h b/Engine/Source/Editor/GameplayDebugger/Public/GameplayDebuggerEdMode.h
new file mode 100644
--- /dev/null
+++ b/Engine/Source/Editor/GameplayDebugger/Public/GameplayDebuggerEdMode.h
@@ -0,0 +1,37 @@
+// Copyright Epic Games, Inc. All Rights Reserved.
+
+#pragma once
+
+#include "CoreMinimal.h"
+
+#include "InputCoreTypes.h"
+#include "EditorModes.h"
+#include "EdMode.h"
+#include "EditorViewportClient.h"
+
+class FGameplayDebuggerEdMode : public FEdMode
+{
+public:
+	FGameplayDebuggerEdMode() : FocusedViewport(nullptr) {}
+
+	// FEdMode interface
+	virtual bool UsesToolkits() const override { return true; }
+	virtual void Enter() override;
+	virtual void Exit() override;
+	virtual bool ReceivedFocus(FEditorViewportClient* ViewportClient, FViewport* Viewport) override;
+	virtual bool LostFocus(FEditorViewportClient* ViewportClient, FViewport* Viewport) override;
+	virtual bool InputKey(FEditorViewportClient* InViewportClient, FViewport* InViewport, FKey InKey, EInputEvent InEvent) override;
+	virtual void Tick(FEditorViewportClient* ViewportClient, float DeltaTime);
+	// End of FEdMode interface
+
+	static void SafeOpenMode();
+	static void SafeCloseMode();
+
+	static const FName EM_GameplayDebugger;
+
+private:
+	FEditorViewportClient* FocusedViewport;
+	bool bPrevGAreScreenMessagesEnabled;
+
+	void EnableViewportClientFlags(FEditorViewportClient* ViewportClient, bool bEnable);
+};
diff --git a/Engine/Source/Editor/GameplayDebugger/Public/GameplayDebuggerEditorModule.h b/Engine/Source/Editor/GameplayDebugger/Public/GameplayDebuggerEditorModule.h
new file mode 100644
--- /dev/null
+++ b/Engine/Source/Editor/GameplayDebugger/Public/GameplayDebuggerEditorModule.h
@@ -0,0 +1,21 @@
+// Copyright Epic Games, Inc. All Rights Reserved.
+
+#pragma once
+
+#include "CoreMinimal.h"
+#include "Engine/World.h"
+#include "Modules/ModuleManager.h"
+#include "Modules/ModuleInterface.h"
+
+class FGameplayDebuggerEditorModule : public IModuleInterface
+{
+public:
+	virtual void StartupModule() override;
+	virtual void ShutdownModule() override;
+
+	void OnWorldInitialized(UWorld* World, const UWorld::InitializationValues IVS);
+
+	void OnLocalControllerInitialized();
+	void OnLocalControllerUninitialized();
+	void OnDebuggerEdModeActivation();
+};
diff --git a/Engine/Source/Editor/GameplayDebugger/Public/GameplayDebuggerExtensionConfigCustomization.h b/Engine/Source/Editor/GameplayDebugger/Public/GameplayDebuggerExtensionConfigCustomization.h
new file mode 100644
--- /dev/null
+++ b/Engine/Source/Editor/GameplayDebugger/Public/GameplayDebuggerExtensionConfigCustomization.h
@@ -0,0 +1,30 @@
+// Copyright Epic Games, Inc. All Rights Reserved.
+
+#pragma once
+
+#include "CoreMinimal.h"
+
+#include "IPropertyTypeCustomization.h"
+
+class IPropertyHandle;
+
+class FGameplayDebuggerExtensionConfigCustomization : public IPropertyTypeCustomization
+{
+public:
+	static TSharedRef<IPropertyTypeCustomization> MakeInstance()
+	{
+		return MakeShareable(new FGameplayDebuggerExtensionConfigCustomization);
+	}
+
+	virtual void CustomizeHeader(TSharedRef<class IPropertyHandle> StructPropertyHandle, class FDetailWidgetRow& HeaderRow, IPropertyTypeCustomizationUtils& StructCustomizationUtils) override;
+	virtual void CustomizeChildren(TSharedRef<class IPropertyHandle> StructPropertyHandle, class IDetailChildrenBuilder& StructBuilder, IPropertyTypeCustomizationUtils& StructCustomizationUtils) override;
+
+protected:
+	TSharedPtr<IPropertyHandle> ExtensionNameProp;
+	TSharedPtr<IPropertyHandle> UseExtensionProp;
+
+	FText CachedHeader;
+	FText GetHeaderDesc() const { return CachedHeader; }
+
+	void OnChildValueChanged();
+};
diff --git a/Engine/Source/Editor/GameplayDebugger/Public/GameplayDebuggerInputConfigCustomization.h b/Engine/Source/Editor/GameplayDebugger/Public/GameplayDebuggerInputConfigCustomization.h
new file mode 100644
--- /dev/null
+++ b/Engine/Source/Editor/GameplayDebugger/Public/GameplayDebuggerInputConfigCustomization.h
@@ -0,0 +1,34 @@
+// Copyright Epic Games, Inc. All Rights Reserved.
+
+#pragma once
+
+#include "CoreMinimal.h"
+
+#include "IPropertyTypeCustomization.h"
+
+class IPropertyHandle;
+
+class FGameplayDebuggerInputConfigCustomization : public IPropertyTypeCustomization
+{
+public:
+	static TSharedRef<IPropertyTypeCustomization> MakeInstance()
+	{
+		return MakeShareable(new FGameplayDebuggerInputConfigCustomization);
+	}
+
+	virtual void CustomizeHeader(TSharedRef<class IPropertyHandle> StructPropertyHandle, class FDetailWidgetRow& HeaderRow, IPropertyTypeCustomizationUtils& StructCustomizationUtils) override;
+	virtual void CustomizeChildren(TSharedRef<class IPropertyHandle> StructPropertyHandle, class IDetailChildrenBuilder& StructBuilder, IPropertyTypeCustomizationUtils& StructCustomizationUtils) override;
+
+protected:
+	TSharedPtr<IPropertyHandle> ConfigNameProp;
+	TSharedPtr<IPropertyHandle> KeyProp;
+	TSharedPtr<IPropertyHandle> ModShiftProp;
+	TSharedPtr<IPropertyHandle> ModCtrlProp;
+	TSharedPtr<IPropertyHandle> ModAltProp;
+	TSharedPtr<IPropertyHandle> ModCmdProp;
+
+	FText CachedHeader;
+	FText GetHeaderDesc() const { return CachedHeader; }
+
+	void OnChildValueChanged();
+};
diff --git a/Engine/Source/Editor/GameplayDebugger/Public/GameplayDebuggerToolkit.h b/Engine/Source/Editor/GameplayDebugger/Public/GameplayDebuggerToolkit.h
new file mode 100644
--- /dev/null
+++ b/Engine/Source/Editor/GameplayDebugger/Public/GameplayDebuggerToolkit.h
@@ -0,0 +1,35 @@
+// Copyright Epic Games, Inc. All Rights Reserved.
+
+#pragma once
+
+#include "CoreMinimal.h"
+
+#include "Layout/Visibility.h"
+#include "Input/Reply.h"
+#include "Widgets/SWidget.h"
+#include "Toolkits/IToolkitHost.h"
+#include "Toolkits/BaseToolkit.h"
+
+class FGameplayDebuggerToolkit : public FModeToolkit
+{
+public:
+	FGameplayDebuggerToolkit(class FEdMode* InOwningMode);
+
+	// IToolkit interface
+	virtual FText GetBaseToolkitName() const override;
+	virtual FName GetToolkitFName() const override;
+	virtual class FEdMode* GetEditorMode() const override { return DebuggerEdMode; }
+	virtual TSharedPtr<class SWidget> GetInlineContent() const override { return MyWidget; }
+	// End of IToolkit interface
+
+	// FModeToolkit interface
+	virtual void Init(const TSharedPtr<class IToolkitHost>& InitToolkitHost) override;
+	// End of FModeToolkit interface
+
+private:
+	class FEdMode* DebuggerEdMode;
+	TSharedPtr<class SWidget> MyWidget;
+
+	EVisibility GetScreenMessageWarningVisibility() const;
+	FReply OnClickedDisableTool();
+};
diff --git a/Engine/Source/Editor/UnrealEd/Private/EditorEngine.cpp b/Engine/Source/Editor/UnrealEd/Private/EditorEngine.cpp
--- a/Engine/Source/Editor/UnrealEd/Private/EditorEngine.cpp
+++ b/Engine/Source/Editor/UnrealEd/Private/EditorEngine.cpp
@@ -1165,6 +1165,7 @@ void UEditorEngine::Init(IEngineLoop* InEngineLoop)
 			TEXT("GeometryFramework"),
 			TEXT("VirtualizationEditor"),
 			TEXT("AnimationSettings"),
+			TEXT("GameplayDebuggerEditor"),
 			TEXT("RenderResourceViewer"),
 		};
 
diff --git a/Engine/Source/Editor/UnrealEd/UnrealEd.Build.cs b/Engine/Source/Editor/UnrealEd/UnrealEd.Build.cs
--- a/Engine/Source/Editor/UnrealEd/UnrealEd.Build.cs
+++ b/Engine/Source/Editor/UnrealEd/UnrealEd.Build.cs
@@ -273,7 +273,8 @@ public class UnrealEd : ModuleRules
 				"WorldPartitionEditor",
 				"CSVtoSVG",
 				"SourceControlWindowExtender",
-				"AnimationSettings"
+				"AnimationSettings",
+				"GameplayDebuggerEditor"
 			}
 		);
 
diff --git a/Engine/Source/Programs/UnrealBuildTool/Configuration/ModuleRules.cs b/Engine/Source/Programs/UnrealBuildTool/Configuration/ModuleRules.cs
--- a/Engine/Source/Programs/UnrealBuildTool/Configuration/ModuleRules.cs
+++ b/Engine/Source/Programs/UnrealBuildTool/Configuration/ModuleRules.cs
@@ -1281,10 +1281,18 @@ namespace UnrealBuildTool
 				if (bAddAsPublicDependency)
 				{
 					PublicDependencyModuleNames.Add("GameplayDebugger");
+					if (Target.Type == TargetType.Editor)
+					{
+						PublicDependencyModuleNames.Add("GameplayDebuggerEditor");
+					}					
 				}
 				else
 				{
 					PrivateDependencyModuleNames.Add("GameplayDebugger");
+					if (Target.Type == TargetType.Editor)
+					{
+						PrivateDependencyModuleNames.Add("GameplayDebuggerEditor");
+					}
 				}
 			}
 			else
diff --git a/Engine/Source/Runtime/GameplayDebugger/GameplayDebugger.Build.cs b/Engine/Source/Runtime/GameplayDebugger/GameplayDebugger.Build.cs
new file mode 100644
--- /dev/null
+++ b/Engine/Source/Runtime/GameplayDebugger/GameplayDebugger.Build.cs
@@ -0,0 +1,55 @@
+// Copyright Epic Games, Inc. All Rights Reserved.
+
+namespace UnrealBuildTool.Rules
+{
+	public class GameplayDebugger : ModuleRules
+	{
+		public GameplayDebugger(ReadOnlyTargetRules Target) : base(Target)
+		{
+			PublicDependencyModuleNames.AddRange(
+				new string[] {
+					"Core",
+					"CoreUObject",
+					"Engine",
+				});
+
+			PrivateDependencyModuleNames.AddRange(
+				new string[] {
+					"RenderCore",
+					"InputCore",
+					"SlateCore",
+					"Slate",
+					"DeveloperSettings",
+				});
+			
+			if (Target.Configuration != UnrealTargetConfiguration.Shipping)
+			{
+				PrivateDependencyModuleNames.Add("DrawPrimitiveDebugger");
+			}
+
+			PrivateIncludePaths.AddRange(
+				new string[] {
+					"Developer/Settings/Public",
+				});
+
+			if (Target.bBuildEditor)
+			{
+				PrivateDependencyModuleNames.AddRange(
+					new string[] {
+						
+						"EditorFramework",
+						"UnrealEd",
+						"LevelEditor",
+						"PropertyEditor",
+					});
+			}
+
+			SetupIrisSupport(Target);
+
+			if (Target.bUseGameplayDebugger)
+			{
+				PrecompileForTargets = PrecompileTargetsType.Any;
+			}
+		}
+	}
+}
diff --git a/Engine/Source/Runtime/GameplayDebugger/Private/GameplayDebuggerAddonBase.cpp b/Engine/Source/Runtime/GameplayDebugger/Private/GameplayDebuggerAddonBase.cpp
new file mode 100644
--- /dev/null
+++ b/Engine/Source/Runtime/GameplayDebugger/Private/GameplayDebuggerAddonBase.cpp
@@ -0,0 +1,67 @@
+// Copyright Epic Games, Inc. All Rights Reserved.
+
+#include "GameplayDebuggerAddonBase.h"
+#include "GameplayDebuggerCategoryReplicator.h"
+#include "GameFramework/PlayerController.h"
+#if WITH_EDITOR
+#include "Editor/EditorEngine.h"
+#endif
+
+AActor* FGameplayDebuggerAddonBase::FindLocalDebugActor() const
+{
+	AGameplayDebuggerCategoryReplicator* RepOwnerOb = RepOwner.Get();
+	return RepOwnerOb ? RepOwnerOb->GetDebugActor() : nullptr;
+}
+
+AGameplayDebuggerCategoryReplicator* FGameplayDebuggerAddonBase::GetReplicator() const
+{
+	return RepOwner.Get();
+}
+
+FString FGameplayDebuggerAddonBase::GetInputHandlerDescription(int32 HandlerId) const
+{
+	return InputHandlers.IsValidIndex(HandlerId) ? InputHandlers[HandlerId].ToString() : FString();
+}
+
+UWorld* FGameplayDebuggerAddonBase::GetWorldFromReplicator() const
+{
+	AGameplayDebuggerCategoryReplicator* Replicator = RepOwner.Get();
+	return Replicator ? Replicator->GetWorld() : nullptr;
+}
+
+UWorld* FGameplayDebuggerAddonBase::GetDataWorld(const APlayerController* OwnerPC, const AActor* DebugActor) const
+{
+	// We're trying OwnerPC first because that's the main owner for what's being displayed by gameplay debugger
+	// but it editor-only mode there are no player controllers, so we try the DebugActor, which can be picked
+	// also in the editor-mode by selecting an actor on the level.
+	// finally, DebugActor can be null as well, when no actor is selected, when we fetch the world which the
+	// AGameplayDebuggerCategoryReplicator instance owning this debugger addon belongs to.
+	UWorld* World = OwnerPC ? OwnerPC->GetWorld() : nullptr;
+	if (World == nullptr && DebugActor != nullptr)
+	{
+		World = DebugActor->GetWorld();
+	}
+	return World ? World : GetWorldFromReplicator();
+}
+
+void FGameplayDebuggerAddonBase::OnGameplayDebuggerActivated()
+{
+	// empty in base class
+}
+
+void FGameplayDebuggerAddonBase::OnGameplayDebuggerDeactivated()
+{
+	// empty in base class
+}
+
+bool FGameplayDebuggerAddonBase::IsSimulateInEditor()
+{
+#if WITH_EDITOR
+	extern UNREALED_API UEditorEngine* GEditor;
+	if (GEditor)
+	{
+		return GEditor->IsSimulateInEditorInProgress();
+	}
+#endif
+	return false;
+}
diff --git a/Engine/Source/Runtime/GameplayDebugger/Private/GameplayDebuggerAddonManager.cpp b/Engine/Source/Runtime/GameplayDebugger/Private/GameplayDebuggerAddonManager.cpp
new file mode 100644
--- /dev/null
+++ b/Engine/Source/Runtime/GameplayDebugger/Private/GameplayDebuggerAddonManager.cpp
@@ -0,0 +1,278 @@
+// Copyright Epic Games, Inc. All Rights Reserved.
+
+#include "GameplayDebuggerAddonManager.h"
+#include "Engine/World.h"
+#include "GameplayDebuggerCategoryReplicator.h"
+#include "GameplayDebuggerExtension.h"
+#include "GameplayDebuggerConfig.h"
+
+FGameplayDebuggerAddonManager::FGameplayDebuggerAddonManager()
+{
+	NumVisibleCategories = 0;
+}
+
+void FGameplayDebuggerAddonManager::RegisterCategory(FName CategoryName, IGameplayDebugger::FOnGetCategory MakeInstanceDelegate, EGameplayDebuggerCategoryState CategoryState, int32 SlotIdx)
+{
+	FGameplayDebuggerCategoryInfo NewInfo;
+	NewInfo.MakeInstanceDelegate = MakeInstanceDelegate;
+	NewInfo.DefaultCategoryState = CategoryState;
+
+	UGameplayDebuggerConfig* MutableToolConfig = UGameplayDebuggerConfig::StaticClass()->GetDefaultObject<UGameplayDebuggerConfig>();
+	uint8 NewCategoryState = (uint8)CategoryState;
+	MutableToolConfig->UpdateCategoryConfig(CategoryName, SlotIdx, NewCategoryState);
+
+	NewInfo.CategoryState = (EGameplayDebuggerCategoryState)NewCategoryState;
+	NewInfo.SlotIdx = SlotIdx;
+
+	CategoryMap.Add(CategoryName, NewInfo);
+
+	// create and destroy single instance to handle input configurators
+	FGameplayDebuggerInputHandlerConfig::CurrentCategoryName = CategoryName;
+	TSharedRef<FGameplayDebuggerCategory> DummyRef = MakeInstanceDelegate.Execute();
+	FGameplayDebuggerInputHandlerConfig::CurrentCategoryName = NAME_None;
+}
+
+void FGameplayDebuggerAddonManager::UnregisterCategory(FName CategoryName)
+{
+	CategoryMap.Remove(CategoryName);
+}
+
+void FGameplayDebuggerAddonManager::NotifyCategoriesChanged()
+{
+	struct FSlotInfo
+	{
+		FName CategoryName;
+		int32 CategoryId;
+		int32 SlotIdx;
+
+		bool operator<(const FSlotInfo& Other) const { return (SlotIdx == Other.SlotIdx) ? CategoryName.LexicalLess(Other.CategoryName) : (SlotIdx < Other.SlotIdx); }
+	};
+
+	TArray<FSlotInfo> AssignList;
+	TSet<int32> OccupiedSlots;
+	int32 CategoryId = 0;
+	for (auto It : CategoryMap)
+	{
+		if (It.Value.CategoryState == EGameplayDebuggerCategoryState::Hidden)
+		{
+			continue;
+		}
+
+		const int32 SanitizedSlotIdx = (It.Value.SlotIdx < 0) ? INDEX_NONE : FMath::Min(100, It.Value.SlotIdx);
+
+		FSlotInfo SlotInfo;
+		SlotInfo.CategoryId = CategoryId;
+		SlotInfo.CategoryName = It.Key;
+		SlotInfo.SlotIdx = SanitizedSlotIdx;
+		CategoryId++;
+
+		AssignList.Add(SlotInfo);
+
+		if (SanitizedSlotIdx >= 0)
+		{
+			OccupiedSlots.Add(SanitizedSlotIdx);
+		}
+	}
+
+	NumVisibleCategories = AssignList.Num();
+	AssignList.Sort();
+
+	int32 MaxSlotIdx = 0;
+	for (int32 Idx = 0; Idx < AssignList.Num(); Idx++)
+	{
+		FSlotInfo& SlotInfo = AssignList[Idx];
+		if (SlotInfo.SlotIdx == INDEX_NONE)
+		{
+			int32 FreeSlotIdx = 0;
+			while (OccupiedSlots.Contains(FreeSlotIdx))
+			{
+				FreeSlotIdx++;
+			}
+
+			SlotInfo.SlotIdx = FreeSlotIdx;
+			OccupiedSlots.Add(FreeSlotIdx);
+		}
+
+		MaxSlotIdx = FMath::Max(MaxSlotIdx, SlotInfo.SlotIdx);
+	}
+
+	SlotMap.Reset();
+	SlotNames.Reset();
+	SlotMap.AddDefaulted(MaxSlotIdx + 1);
+	SlotNames.AddDefaulted(MaxSlotIdx + 1);
+
+	for (int32 Idx = 0; Idx < AssignList.Num(); Idx++)
+	{
+		FSlotInfo& SlotInfo = AssignList[Idx];
+
+		if (SlotNames[SlotInfo.SlotIdx].Len())
+		{
+			SlotNames[SlotInfo.SlotIdx] += TEXT('+');
+		}
+
+		SlotNames[SlotInfo.SlotIdx] += SlotInfo.CategoryName.ToString();
+		SlotMap[SlotInfo.SlotIdx].Add(SlotInfo.CategoryId);
+	}
+
+	OnCategoriesChanged.Broadcast();
+}
+
+void FGameplayDebuggerAddonManager::CreateCategories(AGameplayDebuggerCategoryReplicator& Owner, TArray<TSharedRef<FGameplayDebuggerCategory> >& CategoryObjects)
+{
+	UWorld* World = Owner.GetWorld();
+	check(World);
+	const ENetMode NetMode = World->GetNetMode();
+	const bool bHasAuthority = FGameplayDebuggerUtils::IsAuthority(World);
+	const bool bIsLocal = (NetMode != NM_DedicatedServer);
+	const bool bIsSimulate = FGameplayDebuggerAddonBase::IsSimulateInEditor();
+
+	TArray<TSharedRef<FGameplayDebuggerCategory> > UnsortedCategories;
+	for (auto It : CategoryMap)
+	{
+		FGameplayDebuggerInputHandlerConfig::CurrentCategoryName = It.Key;
+		if (It.Value.CategoryState == EGameplayDebuggerCategoryState::Hidden)
+		{
+			continue;
+		}
+
+		TSharedRef<FGameplayDebuggerCategory> CategoryObjectRef = It.Value.MakeInstanceDelegate.Execute();
+		FGameplayDebuggerCategory& CategoryObject = CategoryObjectRef.Get();
+		CategoryObject.RepOwner = &Owner;
+		CategoryObject.CategoryId = CategoryObjects.Num();
+		CategoryObject.CategoryName = It.Key;
+		CategoryObject.bHasAuthority = bHasAuthority;
+		CategoryObject.bIsLocal = bIsLocal;
+		CategoryObject.bIsEnabled =
+			(It.Value.CategoryState == EGameplayDebuggerCategoryState::EnabledInGameAndSimulate) ||
+			(It.Value.CategoryState == EGameplayDebuggerCategoryState::EnabledInGame && !bIsSimulate) ||
+			(It.Value.CategoryState == EGameplayDebuggerCategoryState::EnabledInSimulate && bIsSimulate);
+
+		UnsortedCategories.Add(CategoryObjectRef);
+	}
+
+	FGameplayDebuggerInputHandlerConfig::CurrentCategoryName = NAME_None;
+
+	// sort by slots for drawing order
+	CategoryObjects.Reset();
+	for (int32 SlotIdx = 0; SlotIdx < SlotMap.Num(); SlotIdx++)
+	{
+		for (int32 Idx = 0; Idx < SlotMap[SlotIdx].Num(); Idx++)
+		{
+			const int32 CategoryId = SlotMap[SlotIdx][Idx];
+			CategoryObjects.Add(UnsortedCategories[CategoryId]);
+		}
+	}
+}
+
+void FGameplayDebuggerAddonManager::RegisterExtension(FName ExtensionName, IGameplayDebugger::FOnGetExtension MakeInstanceDelegate)
+{
+	FGameplayDebuggerExtensionInfo NewInfo;
+	NewInfo.MakeInstanceDelegate = MakeInstanceDelegate;
+	NewInfo.bDefaultEnabled = true;
+
+	uint8 UseExtension = NewInfo.bDefaultEnabled ? 1 : 0;
+	UGameplayDebuggerConfig* MutableToolConfig = UGameplayDebuggerConfig::StaticClass()->GetDefaultObject<UGameplayDebuggerConfig>();
+	MutableToolConfig->UpdateExtensionConfig(ExtensionName, UseExtension);
+
+	NewInfo.bEnabled = UseExtension > 0;
+
+	ExtensionMap.Add(ExtensionName, NewInfo);
+
+	// create and destroy single instance to handle input configurators
+	FGameplayDebuggerInputHandlerConfig::CurrentExtensionName = ExtensionName;
+	TSharedRef<FGameplayDebuggerExtension> DummyRef = MakeInstanceDelegate.Execute();
+	FGameplayDebuggerInputHandlerConfig::CurrentExtensionName = NAME_None;
+}
+
+void FGameplayDebuggerAddonManager::UnregisterExtension(FName ExtensionName)
+{
+	ExtensionMap.Remove(ExtensionName);
+}
+
+void FGameplayDebuggerAddonManager::NotifyExtensionsChanged()
+{
+	OnExtensionsChanged.Broadcast();
+}
+
+void FGameplayDebuggerAddonManager::CreateExtensions(AGameplayDebuggerCategoryReplicator& Replicator, TArray<TSharedRef<FGameplayDebuggerExtension> >& ExtensionObjects)
+{
+	ExtensionObjects.Reset();
+	for (auto It : ExtensionMap)
+	{
+		if (It.Value.bEnabled)
+		{
+			FGameplayDebuggerInputHandlerConfig::CurrentExtensionName = It.Key;
+
+			TSharedRef<FGameplayDebuggerExtension> ExtensionObjectRef = It.Value.MakeInstanceDelegate.Execute();
+			FGameplayDebuggerExtension& ExtensionObject = ExtensionObjectRef.Get();
+			ExtensionObject.RepOwner = &Replicator;
+
+			ExtensionObjects.Add(ExtensionObjectRef);
+		}
+	}
+
+	FGameplayDebuggerInputHandlerConfig::CurrentExtensionName = NAME_None;
+}
+
+void FGameplayDebuggerAddonManager::UpdateFromConfig()
+{
+	UGameplayDebuggerConfig* ToolConfig = UGameplayDebuggerConfig::StaticClass()->GetDefaultObject<UGameplayDebuggerConfig>();
+	if (ToolConfig == nullptr)
+	{
+		return;
+	}
+
+	bool bCategoriesChanged = false;
+	for (auto& It : CategoryMap)
+	{
+		for (int32 Idx = 0; Idx < ToolConfig->Categories.Num(); Idx++)
+		{
+			const FGameplayDebuggerCategoryConfig& ConfigData = ToolConfig->Categories[Idx];
+			if (*ConfigData.CategoryName == It.Key)
+			{
+				const bool bDefaultActiveInGame = (It.Value.DefaultCategoryState == EGameplayDebuggerCategoryState::EnabledInGame) || (It.Value.DefaultCategoryState == EGameplayDebuggerCategoryState::EnabledInGameAndSimulate);
+				const bool bDefaultActiveInSimulate = (It.Value.DefaultCategoryState == EGameplayDebuggerCategoryState::EnabledInSimulate) || (It.Value.DefaultCategoryState == EGameplayDebuggerCategoryState::EnabledInGameAndSimulate);
+
+				const bool bActiveInGame = (ConfigData.ActiveInGame == EGameplayDebuggerOverrideMode::UseDefault) ? bDefaultActiveInGame : (ConfigData.ActiveInGame == EGameplayDebuggerOverrideMode::Enable);
+				const bool bActiveInSimulate = (ConfigData.ActiveInSimulate == EGameplayDebuggerOverrideMode::UseDefault) ? bDefaultActiveInSimulate : (ConfigData.ActiveInSimulate == EGameplayDebuggerOverrideMode::Enable);
+
+				EGameplayDebuggerCategoryState NewCategoryState =
+					bActiveInGame && bActiveInSimulate ? EGameplayDebuggerCategoryState::EnabledInGameAndSimulate :
+					bActiveInGame ? EGameplayDebuggerCategoryState::EnabledInGame :
+					bActiveInSimulate ? EGameplayDebuggerCategoryState::EnabledInSimulate :
+					EGameplayDebuggerCategoryState::Disabled;
+
+				bCategoriesChanged = bCategoriesChanged || (It.Value.SlotIdx != ConfigData.SlotIdx) || (It.Value.CategoryState != NewCategoryState);
+				It.Value.SlotIdx = ConfigData.SlotIdx;
+				It.Value.CategoryState = NewCategoryState;
+				break;
+			}
+		}
+	}
+
+	bool bExtensionsChanged = false;
+	for (auto& It : ExtensionMap)
+	{
+		for (int32 Idx = 0; Idx < ToolConfig->Extensions.Num(); Idx++)
+		{
+			const FGameplayDebuggerExtensionConfig& ConfigData = ToolConfig->Extensions[Idx];
+			if (*ConfigData.ExtensionName == It.Key)
+			{
+				const bool bWantsEnabled = (ConfigData.UseExtension == EGameplayDebuggerOverrideMode::UseDefault) ? It.Value.bDefaultEnabled : (ConfigData.UseExtension == EGameplayDebuggerOverrideMode::Enable);
+				bExtensionsChanged = bExtensionsChanged || (It.Value.bEnabled != bWantsEnabled);
+				It.Value.bEnabled = bWantsEnabled;
+				break;
+			}
+		}
+	}
+
+	if (bCategoriesChanged)
+	{
+		NotifyCategoriesChanged();
+	}
+	
+	if (bExtensionsChanged)
+	{
+		NotifyExtensionsChanged();
+	}
+}
diff --git a/Engine/Source/Runtime/GameplayDebugger/Private/GameplayDebuggerCategory.cpp b/Engine/Source/Runtime/GameplayDebugger/Private/GameplayDebuggerCategory.cpp
new file mode 100644
--- /dev/null
+++ b/Engine/Source/Runtime/GameplayDebugger/Private/GameplayDebuggerCategory.cpp
@@ -0,0 +1,180 @@
+// Copyright Epic Games, Inc. All Rights Reserved.
+
+#include "GameplayDebuggerCategory.h"
+#include "Engine/World.h"
+#include "GameFramework/PlayerController.h"
+#include "GameplayDebuggerCategoryReplicator.h"
+#include "GameplayDebuggerConfig.h"
+#include "GameplayDebuggerPlayerManager.h"
+
+FGameplayDebuggerCategory::FGameplayDebuggerCategory() :
+	CollectDataInterval(0.0f),
+	bShowDataPackReplication(false),
+	bShowUpdateTimer(false),
+	bShowCategoryName(true),
+	bShowOnlyWithDebugActor(true),
+	bIsLocal(false),
+	bHasAuthority(true),
+	bIsEnabled(true),
+	CategoryId(INDEX_NONE),
+	LastCollectDataTime(-FLT_MAX)
+{
+}
+
+FGameplayDebuggerCategory::~FGameplayDebuggerCategory()
+{
+}
+
+void FGameplayDebuggerCategory::CollectData(APlayerController* OwnerPC, AActor* DebugActor)
+{
+	// empty in base class
+}
+
+void FGameplayDebuggerCategory::DrawData(APlayerController* OwnerPC, FGameplayDebuggerCanvasContext& CanvasContext)
+{
+	// empty in base class
+}
+
+FDebugRenderSceneProxy* FGameplayDebuggerCategory::CreateDebugSceneProxy(const UPrimitiveComponent* InComponent, FDebugDrawDelegateHelper*& OutDelegateHelper)
+{
+	OutDelegateHelper = nullptr;
+	// empty in base class
+	return nullptr;
+}
+
+void FGameplayDebuggerCategory::OnDataPackReplicated(int32 DataPackId)
+{
+	// empty in base class
+}
+
+void FGameplayDebuggerCategory::AddTextLine(const FString& TextLine)
+{
+	if (bHasAuthority)
+	{
+		ReplicatedLines.Add(TextLine);
+	}
+}
+
+void FGameplayDebuggerCategory::AddShape(const FGameplayDebuggerShape& Shape)
+{
+	if (bHasAuthority)
+	{
+		ReplicatedShapes.Add(Shape);
+	}
+}
+
+void FGameplayDebuggerCategory::DrawCategory(APlayerController* OwnerPC, FGameplayDebuggerCanvasContext& CanvasContext)
+{
+	// note that we prefer OwnerPC's world here since if it's given then that's the end user of the data
+	UWorld* World = CanvasContext.GetWorld();
+
+	if (!ensure(World))
+	{
+		return;
+	}
+
+	FString CategoryPrefix;
+	if (!bShowCategoryName)
+	{
+		CategoryPrefix = FString::Printf(TEXT("{green}[%s]{white}  "), *CategoryName.ToString());
+	}
+
+	if (bShowUpdateTimer && bHasAuthority)
+	{
+		const float GameTime = World->GetTimeSeconds();
+		CanvasContext.Printf(TEXT("%sNext update in: {yellow}%.0fs"), *CategoryPrefix, CollectDataInterval - (GameTime - LastCollectDataTime));
+	}
+
+	if (bShowDataPackReplication)
+	{
+		for (int32 Idx = 0; Idx < ReplicatedDataPacks.Num(); Idx++)
+		{
+			FGameplayDebuggerDataPack& DataPack = ReplicatedDataPacks[Idx];
+			if (DataPack.IsInProgress())
+			{
+				const FString DataPackMessage = (ReplicatedDataPacks.Num() == 1) ?
+					FString::Printf(TEXT("%sReplicating: {red}%.0f%% {white}(ver:%d)"), *CategoryPrefix, DataPack.GetProgress() * 100.0f, DataPack.Header.DataVersion) :
+					FString::Printf(TEXT("%sReplicating data[%d]: {red}%.0f%% {white}(ver:%d)"), *CategoryPrefix, Idx, DataPack.GetProgress() * 100.0f, DataPack.Header.DataVersion);
+
+				CanvasContext.Print(DataPackMessage);
+			}
+		}
+	}
+
+	for (int32 Idx = 0; Idx < ReplicatedLines.Num(); Idx++)
+	{
+		CanvasContext.Print(ReplicatedLines[Idx]);
+	}
+
+	for (int32 Idx = 0; Idx < ReplicatedShapes.Num(); Idx++)
+	{
+		ReplicatedShapes[Idx].Draw(World, CanvasContext);
+	}
+
+	DrawData(OwnerPC, CanvasContext);
+}
+
+bool FGameplayDebuggerCategory::GetViewPoint(const APlayerController* OwnerPC, FVector& OutViewLocation, FVector& OutViewDirection) const
+{
+	if (const AGameplayDebuggerCategoryReplicator* Replicator = GetReplicator())
+	{
+		if (Replicator->GetViewPoint(OutViewLocation, OutViewDirection))
+		{
+			return true;
+		}
+	}
+
+	if (OwnerPC != nullptr)
+	{
+		AGameplayDebuggerPlayerManager::GetViewPoint(*OwnerPC, OutViewLocation, OutViewDirection);
+		return true;
+	}
+
+	return false;
+}
+
+bool FGameplayDebuggerCategory::IsLocationInViewCone(const FVector& ViewLocation, const FVector& ViewDirection, const FVector& TargetLocation)
+{
+	const UGameplayDebuggerUserSettings* Settings = GetDefault<UGameplayDebuggerUserSettings>();
+	const FVector DirToEntity = TargetLocation - ViewLocation;
+	const float DistanceToEntitySq = DirToEntity.SquaredLength();
+	if (DistanceToEntitySq > FMath::Square(Settings->MaxViewDistance))
+	{
+		return false;
+	}
+
+	const float ViewDot = FVector::DotProduct(DirToEntity.GetSafeNormal(), ViewDirection);
+	const float MinViewDirDot = FMath::Cos(FMath::DegreesToRadians(Settings->MaxViewAngle));
+	if (ViewDot < MinViewDirDot)
+	{
+		return false;
+	}
+
+	return true;
+}
+
+void FGameplayDebuggerCategory::MarkDataPackDirty(int32 DataPackId)
+{
+	if (ReplicatedDataPacks.IsValidIndex(DataPackId))
+	{
+		ReplicatedDataPacks[DataPackId].bIsDirty = true;
+	}
+}
+
+void FGameplayDebuggerCategory::MarkRenderStateDirty()
+{
+	if (bIsLocal)
+	{
+		AGameplayDebuggerCategoryReplicator* RepOwnerOb = GetReplicator();
+		if (RepOwnerOb)
+		{
+			RepOwnerOb->MarkComponentsRenderStateDirty();
+		}
+	}
+}
+
+FString FGameplayDebuggerCategory::GetSceneProxyViewFlag() const
+{
+	const bool bIsSimulate = FGameplayDebuggerAddonBase::IsSimulateInEditor();
+	return bIsSimulate ? TEXT("DebugAI") : TEXT("Game");
+}
diff --git a/Engine/Source/Runtime/GameplayDebugger/Private/GameplayDebuggerCategoryReplicator.cpp b/Engine/Source/Runtime/GameplayDebugger/Private/GameplayDebuggerCategoryReplicator.cpp
new file mode 100644
--- /dev/null
+++ b/Engine/Source/Runtime/GameplayDebugger/Private/GameplayDebuggerCategoryReplicator.cpp
@@ -0,0 +1,1181 @@
+// Copyright Epic Games, Inc. All Rights Reserved.
+
+#include "GameplayDebuggerCategoryReplicator.h"
+#include "Engine/World.h"
+#include "GameFramework/Pawn.h"
+#include "GameFramework/PlayerController.h"
+#include "GameplayDebuggerAddonBase.h"
+#include "GameplayDebuggerAddonManager.h"
+#include "GameplayDebuggerPlayerManager.h"
+#include "GameplayDebuggerRenderingComponent.h"
+#include "GameplayDebuggerExtension.h"
+#include "Net/UnrealNetwork.h"
+#include "VisualLogger/VisualLogger.h"
+
+#if UE_WITH_IRIS
+#include "Iris/IrisConfig.h"
+#include "Net/Iris/ReplicationSystem/ReplicationSystemUtil.h"
+#endif
+
+#include UE_INLINE_GENERATED_CPP_BY_NAME(GameplayDebuggerCategoryReplicator)
+
+#if WITH_EDITOR
+#include "Editor.h"
+#include "Engine/Selection.h"
+#endif // WITH_EDITOR
+
+//////////////////////////////////////////////////////////////////////////
+// FGameplayDebuggerCategoryReplicatorData
+
+DEFINE_LOG_CATEGORY_STATIC(LogGameplayDebugReplication, Display, All);
+
+static TAutoConsoleVariable<int32> CVarGameplayDebuggerRepDetails(
+	TEXT("ai.debug.DetailedReplicationLogs"),
+	0,
+	TEXT("Enable or disable very verbose replication logs for gameplay debugger"),
+	ECVF_Cheat);
+
+static TAutoConsoleVariable<int32> CVarGameplayDebuggerUseDataPackRPC(
+	TEXT("GameplayDebugger.UseDataPackRPC"),
+	0,
+	TEXT("Enable or disable use of rpc's for datapack packets for gameplay debugger"),
+	ECVF_Cheat);
+
+FNotifyGameplayDebuggerOwnerChange AGameplayDebuggerCategoryReplicator::NotifyDebuggerOwnerChange;
+
+class FNetFastCategoryBaseState : public INetDeltaBaseState
+{
+public:
+	struct FDataPackState
+	{
+		int32 DataOffset;
+		int16 DataVersion;
+		int16 SyncCounter;
+
+		FDataPackState() : DataOffset(0), DataVersion(0), SyncCounter(0) {}
+		FDataPackState(const FGameplayDebuggerDataPack::FHeader& Header) : DataOffset(Header.DataOffset), DataVersion(Header.DataVersion), SyncCounter(Header.SyncCounter) {}
+
+		FORCEINLINE bool IsEqual(const FDataPackState& OtherState) const
+		{
+			return (DataOffset == OtherState.DataOffset) && (DataVersion == OtherState.DataVersion) && (SyncCounter == OtherState.SyncCounter);
+		}
+
+		FORCEINLINE bool operator==(const FDataPackState& OtherState) const { return IsEqual(OtherState); }
+		FORCEINLINE bool operator!=(const FDataPackState& OtherState) const { return !IsEqual(OtherState); }
+	};
+
+	struct FCategoryState
+	{
+		int32 TextLinesRepCounter;
+		int32 ShapesRepCounter;
+		TArray<FDataPackState> DataPackStates;
+
+		FCategoryState() : TextLinesRepCounter(0), ShapesRepCounter(0) {}
+
+		FORCEINLINE bool IsEqual(const FCategoryState& OtherState) const
+		{
+			return (TextLinesRepCounter == OtherState.TextLinesRepCounter) &&
+				(ShapesRepCounter == OtherState.ShapesRepCounter) &&
+				(DataPackStates == OtherState.DataPackStates);
+		}
+
+		FORCEINLINE bool operator==(const FCategoryState& OtherState) const { return IsEqual(OtherState); }
+		FORCEINLINE bool operator!=(const FCategoryState& OtherState) const { return !IsEqual(OtherState); }
+	};
+
+	virtual bool IsStateEqual(INetDeltaBaseState* OtherState) override
+	{
+		FNetFastCategoryBaseState* Other = static_cast<FNetFastCategoryBaseState*>(OtherState);
+		return (CategoryStates == Other->CategoryStates);
+	}
+
+	void DumpToLog()
+	{
+		for (int32 CategoryIdx = 0; CategoryIdx < CategoryStates.Num(); CategoryIdx++)
+		{
+			const FCategoryState& CategoryData = CategoryStates[CategoryIdx];
+			UE_LOG(LogGameplayDebugReplication, Verbose, TEXT("category[%d] TextLinesRepCounter:%d ShapesRepCounter:%d"),
+				CategoryIdx, CategoryData.TextLinesRepCounter, CategoryData.ShapesRepCounter);
+
+			for (int32 DataPackIdx = 0; DataPackIdx < CategoryData.DataPackStates.Num(); DataPackIdx++)
+			{
+				const FDataPackState& DataPack = CategoryData.DataPackStates[DataPackIdx];
+				UE_LOG(LogGameplayDebugReplication, Verbose, TEXT(">>    data[%d] DataVersion:%d SyncCounter:%d DataOffset:%d"),
+					DataPackIdx, DataPack.DataVersion, DataPack.SyncCounter, DataPack.DataOffset);
+			}
+		}
+	}
+
+	TArray<FCategoryState> CategoryStates;
+};
+
+void AGameplayDebuggerCategoryReplicator::PreReplication(IRepChangedPropertyTracker& ChangedPropertyTracker)
+{
+#if UE_WITH_IRIS
+	if (UE::Net::FReplicationSystemUtil::GetNetHandle(this).IsValid())
+	{
+		ReplicatedData.PopulateFromOwner();
+	}
+#endif
+}
+
+void AGameplayDebuggerCategoryReplicator::OnRep_ReplicatedData()
+{
+#if UE_WITH_IRIS
+	if (UE::Net::FReplicationSystemUtil::GetNetHandle(this).IsValid())
+	{
+		ReplicatedData.ApplyToOwner();
+	}
+#endif
+}
+
+#if UE_WITH_IRIS
+void FGameplayDebuggerNetPack::PopulateFromOwner()
+{
+	// The current implementation of the FGameplayDebuggerNetPack replication does not behave like most other 
+	// replicated systems so in order to make this work for Iris replication which has much stricter rules for what can be done during serialization,
+	// we cannot rely on polling being made during the call to the custom NetDeltaSerialize method so we use this explicit method to poll the data to be replicated instead.
+	
+	// Update SavedCategories for replication, DataPacks-packets are handled using RPC`s
+	if (Owner && Owner->bIsEnabled && Owner->Categories.Num() == SavedData.Num())
+	{
+		for (int32 Idx = 0; Idx < SavedData.Num(); Idx++)
+		{
+			FGameplayDebuggerCategory& CategoryOb = Owner->Categories[Idx].Get();
+			FGameplayDebuggerCategoryData& SavedCategory = SavedData[Idx];
+
+			SavedCategory.CategoryName = CategoryOb.GetCategoryName();
+			SavedCategory.bIsEnabled = CategoryOb.bIsEnabled;
+
+			const bool bTextLinesChanged = (SavedCategory.TextLines != CategoryOb.ReplicatedLines);
+			if (bTextLinesChanged)
+			{
+				SavedCategory.TextLines = CategoryOb.ReplicatedLines;
+			}
+
+			const bool bShapesChanged = (SavedCategory.Shapes != CategoryOb.ReplicatedShapes);
+			if (bShapesChanged)
+			{
+				SavedCategory.Shapes = CategoryOb.ReplicatedShapes;
+			}
+		}
+	}
+}
+
+void FGameplayDebuggerNetPack::ApplyToOwner()
+{
+	if (!ensure(Owner))
+	{
+		return;
+	}
+
+	bool bHasCategoryStateChanges = false;
+	for (const FGameplayDebuggerCategoryData& SavedDataEntry : SavedData)
+	{
+		// Does the category exist on the client?
+		// using FindLastByPredicate since that's the only Predicate-based find that returns index.
+		const FName CategoryName = SavedDataEntry.CategoryName;
+		const int32 FoundIndex = Owner->Categories.FindLastByPredicate([CategoryName](const TSharedRef<FGameplayDebuggerCategory>& Item) { return (Item->GetCategoryName() == CategoryName); });
+		if (FoundIndex == INDEX_NONE)
+		{
+			continue;
+		}
+		
+		FGameplayDebuggerCategory& CategoryOb = Owner->Categories[FoundIndex].Get();
+		UE_LOG(LogGameplayDebugReplication, Verbose, TEXT("  CATEGORY[%d]:%s"), FoundIndex, *CategoryOb.GetCategoryName().ToString());
+		
+		if (CategoryOb.bIsEnabled != SavedDataEntry.bIsEnabled)
+		{
+			bHasCategoryStateChanges = true;
+			CategoryOb.bIsEnabled = SavedDataEntry.bIsEnabled;
+		}
+
+		if (SavedDataEntry.TextLines != CategoryOb.ReplicatedLines)
+		{
+
+			UE_LOG(LogGameplayDebugReplication, Verbose, TEXT("  >> received lines"));
+			CategoryOb.ReplicatedLines = SavedDataEntry.TextLines;
+		}
+
+		if (SavedDataEntry.Shapes != CategoryOb.ReplicatedShapes)
+		{
+			UE_LOG(LogGameplayDebugReplication, Verbose, TEXT("  >> received shapes"));
+			CategoryOb.ReplicatedShapes = SavedDataEntry.Shapes;
+		}
+	}
+
+	// force scene proxy updates if categories changed state
+	if (bHasCategoryStateChanges)
+	{
+		Owner->MarkComponentsRenderStateDirty();
+	}
+}
+#endif // #if UE_WITH_IRIS
+
+bool FGameplayDebuggerNetPack::NetDeltaSerialize(FNetDeltaSerializeInfo& DeltaParms)
+{
+	if (DeltaParms.bUpdateUnmappedObjects || Owner == nullptr)
+	{
+		return true;
+	}
+
+	if (DeltaParms.Writer)
+	{
+		// watch for DemoNetDriver doing snapshot on client, no need to store debug data here
+		// especially if it tries to write incomplete packets
+		const bool bIsOwnerClient = !Owner->bHasAuthority;
+		if (bIsOwnerClient)
+		{
+			return false;
+		}
+
+		FBitWriter& Writer = *DeltaParms.Writer;
+		int32 NumChangedCategories = 0;
+
+		FNetFastCategoryBaseState* OldState = static_cast<FNetFastCategoryBaseState*>(DeltaParms.OldState);
+		FNetFastCategoryBaseState* NewState = nullptr;
+		TArray<uint8> ChangedCategories;
+
+		// find delta to replicate
+		if (Owner->bIsEnabled && Owner->Categories.Num() == SavedData.Num())
+		{
+			NewState = new FNetFastCategoryBaseState();
+			check(DeltaParms.NewState);
+			*DeltaParms.NewState = TSharedPtr<INetDeltaBaseState>(NewState);
+			NewState->CategoryStates.SetNum(SavedData.Num());
+			ChangedCategories.AddZeroed(SavedData.Num());
+
+			for (int32 Idx = 0; Idx < SavedData.Num(); Idx++)
+			{
+				FNetFastCategoryBaseState::FCategoryState& CategoryState = NewState->CategoryStates[Idx];
+				FGameplayDebuggerCategory& CategoryOb = Owner->Categories[Idx].Get();
+				FGameplayDebuggerCategoryData& SavedCategory = SavedData[Idx];
+
+				const bool bMissingOldState = (OldState == nullptr) || !OldState->CategoryStates.IsValidIndex(Idx);
+				ChangedCategories[Idx] = bMissingOldState ? 1 : 0;
+
+				SavedCategory.CategoryName = CategoryOb.GetCategoryName();
+				if (SavedCategory.bIsEnabled != CategoryOb.bIsEnabled)
+				{
+					SavedCategory.bIsEnabled = CategoryOb.bIsEnabled;
+					ChangedCategories[Idx]++;
+				}
+
+				const bool bTextLinesChanged = (SavedCategory.TextLines != CategoryOb.ReplicatedLines);
+				CategoryState.TextLinesRepCounter = (bMissingOldState ? 0 : OldState->CategoryStates[Idx].TextLinesRepCounter) + (bTextLinesChanged ? 1 : 0);
+				if (bTextLinesChanged)
+				{
+					SavedCategory.TextLines = CategoryOb.ReplicatedLines;
+					ChangedCategories[Idx]++;
+				}
+
+				const bool bShapesChanged = (SavedCategory.Shapes != CategoryOb.ReplicatedShapes);
+				CategoryState.ShapesRepCounter = (bMissingOldState ? 0 : OldState->CategoryStates[Idx].ShapesRepCounter) + (bShapesChanged ? 1 : 0);
+				if (bShapesChanged)
+				{
+					SavedCategory.Shapes = CategoryOb.ReplicatedShapes;
+					ChangedCategories[Idx]++;
+				}
+
+				// If datapack-packets are sent using RPC`s we do not need to store them in the SavedCategoryStates or the Delta baselines
+				const bool bUseDataPackRPC = Owner->bSendDataPacksUsingRPC;
+				const int32 NumDataPacks = bUseDataPackRPC ? 0 : CategoryOb.ReplicatedDataPacks.Num();
+
+				SavedCategory.DataPacks.SetNum(NumDataPacks);
+				CategoryState.DataPackStates.SetNum(NumDataPacks);
+				for (int32 DataIdx = 0; DataIdx < NumDataPacks; DataIdx++)
+				{
+					FGameplayDebuggerDataPack& DataPack = CategoryOb.ReplicatedDataPacks[DataIdx];
+					const bool bHasOldStatePack = !bMissingOldState && OldState->CategoryStates[Idx].DataPackStates.IsValidIndex(DataIdx);
+
+					if (DataPack.bNeedsConfirmation && !DataPack.bReceived && bHasOldStatePack)
+					{
+						FNetFastCategoryBaseState::FDataPackState& OldDataPackState = OldState->CategoryStates[Idx].DataPackStates[DataIdx];
+						UE_LOG(LogGameplayDebugReplication, Verbose, TEXT("Checking packet confirmation for Category[%d].DataPack[%d] OldState(DataVersion:%d DataOffset:%d complete:%s) current(DataVersion:%d DataOffset:%d)"),
+							Idx, DataIdx,
+							OldDataPackState.DataVersion, OldDataPackState.DataOffset,
+							(OldDataPackState.DataOffset == DataPack.Header.DataSize) && (OldDataPackState.DataVersion == DataPack.Header.DataVersion) ? TEXT("yes") : TEXT("no"),
+							DataPack.Header.DataVersion, DataPack.Header.DataOffset);
+
+						DataPack.OnPacketRequest(OldDataPackState.DataVersion, OldDataPackState.DataOffset);
+					}
+
+					CategoryState.DataPackStates[DataIdx] = FNetFastCategoryBaseState::FDataPackState(DataPack.Header);
+					const bool bDataPackChanged = (SavedCategory.DataPacks[DataIdx] != DataPack.Header);
+					if (bDataPackChanged)
+					{
+						SavedCategory.DataPacks[DataIdx] = DataPack.Header;
+						ChangedCategories[Idx]++;
+					}
+					else if (bHasOldStatePack)
+					{
+						FNetFastCategoryBaseState::FDataPackState& OldDataPackState = OldState->CategoryStates[Idx].DataPackStates[DataIdx];
+						const bool bDataPackNotUpdatedOnClient = (OldDataPackState != DataPack.Header);
+						if (bDataPackNotUpdatedOnClient)
+						{
+							ChangedCategories[Idx]++;
+						}
+					}
+				}
+
+				NumChangedCategories += ChangedCategories[Idx] ? 1 : 0;
+			}
+		}
+
+		if (CVarGameplayDebuggerRepDetails.GetValueOnAnyThread())
+		{
+			if (OldState)
+			{
+				UE_LOG(LogGameplayDebugReplication, Verbose, TEXT("NetDeltaSerialize DUMP OldState"));
+				OldState->DumpToLog();
+			}
+			if (NewState)
+			{
+				UE_LOG(LogGameplayDebugReplication, Verbose, TEXT("NetDeltaSerialize DUMP NewState"));
+				NewState->DumpToLog();
+			}
+		}
+
+		if (NumChangedCategories == 0)
+		{
+			return false;
+		}
+
+		int32 CategoryCount = SavedData.Num();
+		Writer << CategoryCount;
+
+		UE_LOG(LogGameplayDebugReplication, Verbose, TEXT("NetDeltaSerialize WRITE START, OldState:%d"), OldState ? 1 : 0);
+		for (int32 Idx = 0; Idx < SavedData.Num(); Idx++)
+		{
+			FGameplayDebuggerCategory& CategoryOb = Owner->Categories[Idx].Get();
+			const bool bMissingOldState = (OldState == nullptr) || !OldState->CategoryStates.IsValidIndex(Idx);
+			FGameplayDebuggerCategoryData& SavedCategory = SavedData[Idx];
+
+			uint8 BaseFlags = SavedCategory.bIsEnabled;
+			uint8 ShouldUpdateTextLines = bMissingOldState || (OldState->CategoryStates[Idx].TextLinesRepCounter != NewState->CategoryStates[Idx].TextLinesRepCounter);
+			uint8 ShouldUpdateShapes = bMissingOldState || (OldState->CategoryStates[Idx].ShapesRepCounter != NewState->CategoryStates[Idx].ShapesRepCounter);
+			uint8 NumDataPacks = SavedCategory.DataPacks.Num();
+
+			Writer << SavedCategory.CategoryName;
+
+			Writer.WriteBit(BaseFlags);
+			Writer.WriteBit(ShouldUpdateTextLines);
+			Writer.WriteBit(ShouldUpdateShapes);
+			Writer << NumDataPacks;
+
+			if (ChangedCategories[Idx])
+			{
+				UE_LOG(LogGameplayDebugReplication, Verbose, TEXT("  CATEGORY[%d]:%s"), Idx, *CategoryOb.GetCategoryName().ToString());
+			}
+
+			if (ShouldUpdateTextLines)
+			{
+				UE_LOG(LogGameplayDebugReplication, Verbose, TEXT("  >> replicate lines"));
+				Writer << SavedCategory.TextLines;
+			}
+
+			if (ShouldUpdateShapes)
+			{
+				UE_LOG(LogGameplayDebugReplication, Verbose, TEXT("  >> replicate shapes"));
+				Writer << SavedCategory.Shapes;
+			}
+
+			for (int32 DataIdx = 0; DataIdx < NumDataPacks; DataIdx++)
+			{
+				uint8 ShouldUpdateDataPack = bMissingOldState || !OldState->CategoryStates[Idx].DataPackStates.IsValidIndex(DataIdx) ||
+					(OldState->CategoryStates[Idx].DataPackStates[DataIdx] != NewState->CategoryStates[Idx].DataPackStates[DataIdx]);
+
+				Writer.WriteBit(ShouldUpdateDataPack);
+				if (ShouldUpdateDataPack)
+				{
+					// send single packet
+					FGameplayDebuggerDataPack& DataPack = CategoryOb.ReplicatedDataPacks[DataIdx];
+
+					uint8 IsCompressed = DataPack.Header.bIsCompressed ? 1 : 0;
+					Writer.WriteBit(IsCompressed);
+
+					Writer << DataPack.Header.DataVersion;
+					Writer << DataPack.Header.SyncCounter;
+					Writer << DataPack.Header.DataSize;
+					Writer << DataPack.Header.DataOffset;
+
+					const int32 PacketSize = FMath::Min(FGameplayDebuggerDataPack::PacketSize, DataPack.Header.DataSize - DataPack.Header.DataOffset);
+					if (PacketSize > 0)
+					{
+						Writer.Serialize(DataPack.Data.GetData() + DataPack.Header.DataOffset, PacketSize);
+					}
+
+					UE_LOG(LogGameplayDebugReplication, Verbose, TEXT("  >> replicate data pack[%d] progress:%0.f%% (offset:%d packet:%d)"),
+						DataIdx, DataPack.Header.DataSize ? (100.0f * (DataPack.Header.DataOffset + PacketSize) / DataPack.Header.DataSize) : 100.0f,
+						DataPack.Header.DataOffset, PacketSize);
+				}
+			}
+		}
+	}
+	else if (DeltaParms.Reader)
+	{
+		FBitReader& Reader = *DeltaParms.Reader;
+		UE_LOG(LogGameplayDebugReplication, Verbose, TEXT("NetDeltaSerialize READ START"));
+
+		int32 CategoryCount = 0;
+		Reader << CategoryCount;
+
+		// Replicated categories still have to be serialized even if they don't exist on the client.
+		// Create a dummy category for these cases
+		FGameplayDebuggerCategory DummyCategory;
+
+		bool bHasCategoryStateChanges = false;
+		for (int32 Idx = 0; Idx < CategoryCount; Idx++)
+		{
+			FName CategoryName;
+			Reader << CategoryName;
+
+			// Does the category exist on the client?
+			// using FindLastByPredicate since that's the only Predicate-based find that returns index.
+			const int32 FoundIndex = Owner->Categories.FindLastByPredicate([CategoryName](const TSharedRef<FGameplayDebuggerCategory>& Item)
+				{
+					return (Item->GetCategoryName() == CategoryName);
+				});
+
+			FGameplayDebuggerCategory& CategoryOb = (Owner->Categories.IsValidIndex(FoundIndex)) ? Owner->Categories[FoundIndex].Get() : DummyCategory;
+			UE_LOG(LogGameplayDebugReplication, Verbose, TEXT("  CATEGORY[%d]:%s"), FoundIndex, *CategoryOb.GetCategoryName().ToString());
+
+			uint8 BaseFlags = Reader.ReadBit();
+			uint8 ShouldUpdateTextLines = Reader.ReadBit();
+			uint8 ShouldUpdateShapes = Reader.ReadBit();
+
+			uint8 NumDataPacks = 0;
+			Reader << NumDataPacks;
+
+			if (FoundIndex != INDEX_NONE && (int32)NumDataPacks != CategoryOb.ReplicatedDataPacks.Num())
+			{
+				UE_LOG(LogGameplayDebugReplication, Error, TEXT("Data pack count mismatch! received:%d expected:%d"), NumDataPacks, CategoryOb.ReplicatedDataPacks.Num());
+				Reader.SetError();
+				return false;
+			}
+
+			const bool bNewCategoryEnabled = (BaseFlags != 0);
+			bHasCategoryStateChanges = bHasCategoryStateChanges || (CategoryOb.bIsEnabled != bNewCategoryEnabled);
+			CategoryOb.bIsEnabled = bNewCategoryEnabled;
+
+			if (ShouldUpdateTextLines)
+			{
+				UE_LOG(LogGameplayDebugReplication, Verbose, TEXT("  >> received lines"));
+				Reader << CategoryOb.ReplicatedLines;
+			}
+
+			if (ShouldUpdateShapes)
+			{
+				UE_LOG(LogGameplayDebugReplication, Verbose, TEXT("  >> received shapes"));
+				Reader << CategoryOb.ReplicatedShapes;
+			}
+
+			for (uint8 DataIdx = 0; DataIdx < NumDataPacks; DataIdx++)
+			{
+				uint8 ShouldUpdateDataPack = Reader.ReadBit();
+
+				if (ShouldUpdateDataPack)
+				{
+					// receive single packet
+					FGameplayDebuggerDataPack DataPacket;
+
+					uint8 IsCompressed = Reader.ReadBit();
+					DataPacket.Header.bIsCompressed = (IsCompressed != 0);
+
+					Reader << DataPacket.Header.DataVersion;
+					Reader << DataPacket.Header.SyncCounter;
+					Reader << DataPacket.Header.DataSize;
+					Reader << DataPacket.Header.DataOffset;
+
+					const int32 PacketSize = FMath::Min(FGameplayDebuggerDataPack::PacketSize, DataPacket.Header.DataSize - DataPacket.Header.DataOffset);
+					if (PacketSize > 0)
+					{
+						DataPacket.Data.AddUninitialized(PacketSize);
+						Reader.Serialize(DataPacket.Data.GetData(), PacketSize);
+					}
+
+					Owner->OnReceivedDataPackPacket(FoundIndex, DataIdx, DataPacket);
+					UE_LOG(LogGameplayDebugReplication, Verbose, TEXT("  >> replicate data pack[%d] progress:%.0f%%"), DataIdx, CategoryOb.ReplicatedDataPacks[DataIdx].GetProgress() * 100.0f);
+				}
+			}
+		}
+
+		// force scene proxy updates if categories changed state
+		if (bHasCategoryStateChanges)
+		{
+			Owner->MarkComponentsRenderStateDirty();
+		}
+	}
+
+	return true;
+}
+
+void FGameplayDebuggerNetPack::OnCategoriesChanged()
+{
+	SavedData.Reset();
+	SavedData.SetNum(Owner->Categories.Num());
+}
+
+//////////////////////////////////////////////////////////////////////////
+// AGameplayDebuggerCategoryReplicator
+
+AGameplayDebuggerCategoryReplicator::AGameplayDebuggerCategoryReplicator(const FObjectInitializer& ObjectInitializer) : Super(ObjectInitializer)
+{
+	PrimaryActorTick.bCanEverTick = true;
+	PrimaryActorTick.bAllowTickOnDedicatedServer = true;
+	PrimaryActorTick.bTickEvenWhenPaused = true;
+	PrimaryActorTick.bStartWithTickEnabled = false;
+
+#if WITH_EDITOR
+	SetIsTemporarilyHiddenInEditor(true);
+#endif
+
+#if WITH_EDITORONLY_DATA
+	bHiddenEdLevel = true;
+	bHiddenEdLayer = true;
+	bHiddenEd = true;
+	bEditable = false;
+#endif
+
+	bIsEnabled = false;
+	bIsEnabledLocal = false;
+	bHasAuthority = false;
+	bIsLocal = false;
+	bIsEditorWorldReplicator = false;
+	bReplicates = true;
+
+#if UE_WITH_IRIS
+	if (UE::Net::ShouldUseIrisReplication())
+	{
+		// If iris is enabled, datapack-packets are always sent using RPC`s and we also requires PreReplication to called in order to populate replicated data outside of serialization.
+		SetCallPreReplication(true);
+		bSendDataPacksUsingRPC = true;
+	}
+#else
+	// We cache this as it does not make sense to be able to toggle it other than between sessions.
+	bSendDataPacksUsingRPC = CVarGameplayDebuggerUseDataPackRPC.GetValueOnAnyThread();
+#endif
+
+	ReplicatedData.Owner = this;
+}
+
+void AGameplayDebuggerCategoryReplicator::BeginPlay()
+{
+	Super::BeginPlay();
+
+	UWorld* World = GetWorld();
+	check(World);
+	const ENetMode NetMode = World->GetNetMode();
+	bHasAuthority = FGameplayDebuggerUtils::IsAuthority(World);
+	bIsLocal = (NetMode != NM_DedicatedServer);
+
+	Init();
+}
+
+#if UE_WITH_IRIS
+void AGameplayDebuggerCategoryReplicator::BeginReplication()
+{
+	bOnlyRelevantToOwner = true;
+	Super::BeginReplication();
+}
+
+const AActor* AGameplayDebuggerCategoryReplicator::GetNetOwner() const 
+{
+	return OwnerPC.Get();
+}
+#endif
+
+#if WITH_EDITOR
+void AGameplayDebuggerCategoryReplicator::InitForEditor()
+{
+	bHasAuthority = true;
+	bIsLocal = true;
+	bIsEditorWorldReplicator = true;
+	Init();
+}
+#endif // WITH_EDITOR
+
+void AGameplayDebuggerCategoryReplicator::Init()
+{
+	UWorld* World = GetWorld();
+	check(World);
+
+	FGameplayDebuggerAddonManager& AddonManager = FGameplayDebuggerAddonManager::GetCurrent();
+	AddonManager.OnCategoriesChanged.AddUObject(this, &AGameplayDebuggerCategoryReplicator::OnCategoriesChanged);
+	AddonManager.OnExtensionsChanged.AddUObject(this, &AGameplayDebuggerCategoryReplicator::OnExtensionsChanged);
+
+	OnCategoriesChanged();
+	OnExtensionsChanged();
+
+	SetActorHiddenInGame(!bIsLocal);
+	if (bIsLocal)
+	{
+		RenderingComp = NewObject<UGameplayDebuggerRenderingComponent>(this, TEXT("RenderingComp"));
+		RenderingComp->RegisterComponentWithWorld(World);
+		RootComponent = RenderingComp;
+	}
+	
+	if (bHasAuthority)
+	{
+		SetEnabled(FGameplayDebuggerAddonBase::IsSimulateInEditor());
+	}
+
+	AGameplayDebuggerPlayerManager& PlayerManager = AGameplayDebuggerPlayerManager::GetCurrent(World);
+	PlayerManager.RegisterReplicator(*this);
+}
+
+void AGameplayDebuggerCategoryReplicator::EndPlay(const EEndPlayReason::Type EndPlayReason)
+{
+	Super::EndPlay(EndPlayReason);
+
+	SetReplicatorOwner(nullptr);
+
+	// Disable extensions to clear UI state
+	NotifyCategoriesToolState(false);
+	NotifyExtensionsToolState(false);
+
+	FGameplayDebuggerAddonManager& AddonManager = FGameplayDebuggerAddonManager::GetCurrent();
+	AddonManager.OnCategoriesChanged.RemoveAll(this);
+	AddonManager.OnExtensionsChanged.RemoveAll(this);
+}
+
+void AGameplayDebuggerCategoryReplicator::OnCategoriesChanged()
+{
+	FGameplayDebuggerAddonManager& AddonManager = FGameplayDebuggerAddonManager::GetCurrent();
+	AddonManager.CreateCategories(*this, Categories);
+
+	ReplicatedData.OnCategoriesChanged();
+
+	if (bIsLocal)
+	{
+		AGameplayDebuggerPlayerManager& PlayerManager = AGameplayDebuggerPlayerManager::GetCurrent(GetWorld());
+		PlayerManager.RefreshInputBindings(*this);
+	}
+}
+
+void AGameplayDebuggerCategoryReplicator::OnExtensionsChanged()
+{
+	FGameplayDebuggerAddonManager& AddonManager = FGameplayDebuggerAddonManager::GetCurrent();
+	AddonManager.CreateExtensions(*this, Extensions);
+
+	if (bIsLocal)
+	{
+		AGameplayDebuggerPlayerManager& PlayerManager = AGameplayDebuggerPlayerManager::GetCurrent(GetWorld());
+		PlayerManager.RefreshInputBindings(*this);
+	}
+}
+
+UNetConnection* AGameplayDebuggerCategoryReplicator::GetNetConnection() const
+{
+	return IsValid(OwnerPC) ? OwnerPC->GetNetConnection() : nullptr;
+}
+
+bool AGameplayDebuggerCategoryReplicator::IsNetRelevantFor(const AActor* RealViewer, const AActor* ViewTarget, const FVector& SrcLocation) const
+{
+	return (RealViewer == OwnerPC);
+}
+
+/// @cond DOXYGEN_WARNINGS
+
+void AGameplayDebuggerCategoryReplicator::GetLifetimeReplicatedProps(TArray<FLifetimeProperty>& OutLifetimeProps) const
+{
+	Super::GetLifetimeReplicatedProps(OutLifetimeProps);
+
+	DOREPLIFETIME(AGameplayDebuggerCategoryReplicator, OwnerPC);
+	DOREPLIFETIME(AGameplayDebuggerCategoryReplicator, DebugActor);
+	DOREPLIFETIME(AGameplayDebuggerCategoryReplicator, VisLogSync);
+	DOREPLIFETIME(AGameplayDebuggerCategoryReplicator, bIsEnabled);
+	DOREPLIFETIME(AGameplayDebuggerCategoryReplicator, ReplicatedData);
+
+	// The visibility of the replicator actor is code driven (hidden on dedicated server, visible otherwise) and should not be replicated
+	DISABLE_REPLICATED_PRIVATE_PROPERTY(AActor, bHidden);
+}
+
+bool AGameplayDebuggerCategoryReplicator::ServerSetEnabled_Validate(bool bEnable)
+{
+	return true;
+}
+
+void AGameplayDebuggerCategoryReplicator::ServerSetEnabled_Implementation(bool bEnable)
+{
+	SetEnabled(bEnable);
+}
+
+bool AGameplayDebuggerCategoryReplicator::ServerSetDebugActor_Validate(AActor* Actor, bool bSelectInEditor)
+{
+	return true;
+}
+
+void AGameplayDebuggerCategoryReplicator::ServerSetDebugActor_Implementation(AActor* Actor, bool bSelectInEditor)
+{
+	SetDebugActor(Actor, bSelectInEditor);
+}
+
+bool AGameplayDebuggerCategoryReplicator::ServerSetViewPoint_Validate(const FVector& InViewLocation, const FVector& InViewDirection)
+{
+	return true;
+}
+
+void AGameplayDebuggerCategoryReplicator::ServerSetViewPoint_Implementation(const FVector& InViewLocation, const FVector& InViewDirection)
+{
+	SetViewPoint(InViewLocation, InViewDirection);
+}
+
+bool AGameplayDebuggerCategoryReplicator::ServerResetViewPoint_Validate()
+{
+	return true;
+}
+
+void AGameplayDebuggerCategoryReplicator::ServerResetViewPoint_Implementation()
+{
+	ResetViewPoint();
+}
+
+bool AGameplayDebuggerCategoryReplicator::ServerSetCategoryEnabled_Validate(int32 CategoryId, bool bEnable)
+{
+	return true;
+}
+
+void AGameplayDebuggerCategoryReplicator::ServerSetCategoryEnabled_Implementation(int32 CategoryId, bool bEnable)
+{
+	SetCategoryEnabled(CategoryId, bEnable);
+}
+
+bool AGameplayDebuggerCategoryReplicator::ServerSendCategoryInputEvent_Validate(int32 CategoryId, int32 HandlerId)
+{
+	return true;
+}
+
+void AGameplayDebuggerCategoryReplicator::ServerSendCategoryInputEvent_Implementation(int32 CategoryId, int32 HandlerId)
+{
+	SendCategoryInputEvent(CategoryId, HandlerId);
+}
+
+bool AGameplayDebuggerCategoryReplicator::ServerSendExtensionInputEvent_Validate(int32 ExtensionId, int32 HandlerId)
+{
+	return true;
+}
+
+void AGameplayDebuggerCategoryReplicator::ServerSendExtensionInputEvent_Implementation(int32 ExtensionId, int32 HandlerId)
+{
+	SendExtensionInputEvent(ExtensionId, HandlerId);
+}
+
+/// @endcond
+
+void AGameplayDebuggerCategoryReplicator::OnReceivedDataPackPacket(int32 CategoryId, int32 DataPackId, const FGameplayDebuggerDataPack& DataPacket)
+{
+	if (Categories.IsValidIndex(CategoryId) && Categories[CategoryId]->ReplicatedDataPacks.IsValidIndex(DataPackId))
+	{
+		FGameplayDebuggerDataPack& DataPack = Categories[CategoryId]->ReplicatedDataPacks[DataPackId];
+		bool bIsPacketValid = false;
+
+		if (DataPack.Header.DataVersion != DataPacket.Header.DataVersion)
+		{
+			// new content of data pack:
+			if (DataPacket.Header.DataOffset == 0)
+			{
+				// first packet of data, replace old data pack's intermediate data
+				DataPack.Header = DataPacket.Header;
+				DataPack.Data = DataPacket.Data;
+				bIsPacketValid = true;
+			}
+			else
+			{
+				// somewhere in the middle, discard packet
+				UE_LOG(LogGameplayDebugReplication, Verbose, TEXT("Error: received packet from the middle of content with different version, discarding! (category[%d]:%s, packet.DataVersion:%d packet.DataOffset:%d, data[%d].DataVersion:%d)"),
+					CategoryId, *Categories[CategoryId]->GetCategoryName().ToString(),
+					DataPacket.Header.DataVersion,
+					DataPacket.Header.DataOffset,
+					DataPackId, DataPack.Header.DataVersion,
+					DataPackId, DataPack.Header.DataOffset);
+			}
+		}
+		else if (DataPack.Data.Num() < DataPacket.Header.DataSize)
+		{
+			// another packet for existing data pack
+			if (DataPacket.Header.DataOffset == DataPack.Data.Num())
+			{
+				// offset matches, this is next expected packet
+				DataPack.Data.Append(DataPacket.Data);
+				DataPack.Header.DataOffset = DataPack.Data.Num();
+				bIsPacketValid = true;
+			}
+			else
+			{
+				// offset mismatch, discard packet
+				UE_LOG(LogGameplayDebugReplication, Verbose, TEXT("Error: received packet doesn't match expected chunk, discarding! (category[%d]:%s, packet.DataOffset:%d, data[%d].DataOffset:%d data[%d].Data.Num:%d)"),
+					CategoryId, *Categories[CategoryId]->GetCategoryName().ToString(),
+					DataPacket.Header.DataOffset,
+					DataPackId, DataPack.Header.DataOffset,
+					DataPackId, DataPack.Data.Num());
+			}
+		}
+
+		// check if data pack is now complete
+		if (bIsPacketValid && (DataPack.Data.Num() == DataPack.Header.DataSize))
+		{
+			// complete
+			UE_LOG(LogGameplayDebugReplication, Verbose, TEXT("Category[%d].DataPack[%d] RECEIVED, DataVersion:%d DataSize:%d SyncCounter:%d"),
+				CategoryId, DataPackId, DataPack.Header.DataVersion, DataPack.Header.DataSize, DataPack.Header.SyncCounter);
+
+			DataPack.OnReplicated();
+			Categories[CategoryId]->OnDataPackReplicated(DataPackId);
+		}
+	}
+}
+
+void AGameplayDebuggerCategoryReplicator::TickActor(float DeltaTime, enum ELevelTick TickType, FActorTickFunction& ThisTickFunction)
+{
+	Super::TickActor(DeltaTime, TickType, ThisTickFunction);
+	if (OwnerPC || bIsEditorWorldReplicator)
+	{
+		CollectCategoryData();
+	}
+}
+
+void AGameplayDebuggerCategoryReplicator::PostNetReceive()
+{
+	Super::PostNetReceive();
+
+	// force scene proxy updates if tool changed state
+	if (bIsEnabled != bIsEnabledLocal)
+	{
+		bIsEnabledLocal = bIsEnabled;
+		MarkComponentsRenderStateDirty();
+	}
+}
+
+void AGameplayDebuggerCategoryReplicator::CollectCategoryData(bool bForce)
+{
+	const float GameTime = GetWorld()->GetTimeSeconds();
+
+#if ENABLE_VISUAL_LOG
+	const FVisualLogger& VLogger = FVisualLogger::Get();
+	VisLogSync.DeviceIDs.Reset();
+	if (VLogger.IsRecordingToFile())
+	{
+		const TArray<FVisualLogDevice*>& LogDevices = VLogger.GetDevices();
+		for (const FVisualLogDevice* Device : LogDevices)
+		{
+			VisLogSync.DeviceIDs += FString::Printf(TEXT("%u, "), Device ? Device->GetShortSessionID() : 0);
+		}
+	}
+#endif // ENABLE_VISUAL_LOG
+
+	for (int32 Idx = 0; Idx < Categories.Num(); Idx++)
+	{
+		FGameplayDebuggerCategory& CategoryOb = Categories[Idx].Get();
+		if (CategoryOb.bHasAuthority && CategoryOb.bIsEnabled && (bForce || (GameTime - CategoryOb.LastCollectDataTime) > CategoryOb.CollectDataInterval))
+		{
+			// prepare data packs before calling CollectData
+			for (int32 DataPackIdx = 0; DataPackIdx < CategoryOb.ReplicatedDataPacks.Num(); DataPackIdx++)
+			{
+				FGameplayDebuggerDataPack& DataPack = CategoryOb.ReplicatedDataPacks[DataPackIdx];
+				DataPack.bIsDirty = false;
+
+				if ((DataPack.Flags == EGameplayDebuggerDataPack::ResetOnTick) ||
+					(DataPack.Flags == EGameplayDebuggerDataPack::ResetOnActorChange && DataPack.Header.SyncCounter != DebugActor.SyncCounter))
+				{
+					DataPack.ResetDelegate.Execute();
+				}
+			}
+
+			CategoryOb.ReplicatedLines.Reset();
+			CategoryOb.ReplicatedShapes.Reset();
+
+			CategoryOb.CollectData(OwnerPC, DebugActor.Actor.Get());
+			CategoryOb.LastCollectDataTime = GameTime;
+
+			// update dirty data packs
+			for (int32 DataPackIdx = 0; DataPackIdx < CategoryOb.ReplicatedDataPacks.Num(); DataPackIdx++)
+			{
+				FGameplayDebuggerDataPack& DataPack = CategoryOb.ReplicatedDataPacks[DataPackIdx];
+				if (CategoryOb.bIsLocal)
+				{
+					const bool bWasDirty = DataPack.CheckDirtyAndUpdate();
+					if (bWasDirty)
+					{
+						CategoryOb.OnDataPackReplicated(DataPackIdx);
+					}
+
+					if (CategoryOb.bHasAuthority)
+					{
+						// update sync counter for local & auth packs (no data replication), otherwise they can be reset
+						DataPack.Header.SyncCounter = DebugActor.SyncCounter;
+					}
+				}
+				else
+				{
+					const bool bWasDirty = DataPack.RequestReplication(DebugActor.SyncCounter);
+					if (bWasDirty)
+					{
+						UE_LOG(LogGameplayDebugReplication, Verbose, TEXT("Category[%d].DataPack[%d] SENT, DataVersion:%d DataSize:%d SyncCounter:%d"),
+							Idx, DataPackIdx, DataPack.Header.DataVersion, DataPack.Header.DataSize, DataPack.Header.SyncCounter);
+
+						// Send the update data pack as an reliable rpc instead
+						if (bSendDataPacksUsingRPC)
+						{
+							SendDataPackPacket(CategoryOb.GetCategoryName(), DataPackIdx, DataPack);						
+						}
+					}
+				}
+			}
+		}
+	}
+}
+
+void AGameplayDebuggerCategoryReplicator::SendDataPackPacket(FName CategoryName, int32 DataPackIdx, FGameplayDebuggerDataPack& DataPack)
+{
+	FGameplayDebuggerDataPackRPCParams Params;
+
+	Params.CategoryName = CategoryName;
+	Params.DataPackIdx = DataPackIdx;
+	Params.Header = DataPack.Header;
+
+	const int32 PacketSize = FMath::Min(FGameplayDebuggerDataPack::PacketSize, DataPack.Header.DataSize - DataPack.Header.DataOffset);
+	if (PacketSize > 0)
+	{
+		Params.Data.Append(DataPack.Data.GetData() + DataPack.Header.DataOffset, PacketSize);
+	}
+
+	UE_LOG(LogGameplayDebugReplication, Verbose, TEXT("  >> replicate data pack[%d] RPC progress:%0.f%% (offset:%d packet:%d)"),
+		DataPackIdx, DataPack.Header.DataSize ? (100.0f * (DataPack.Header.DataOffset + PacketSize) / DataPack.Header.DataSize) : 100.0f,
+		DataPack.Header.DataOffset, PacketSize);
+
+	ClientDataPackPacket(Params);
+
+	// As we are sending this as a reliable RPC we "ack" it immediately
+	if (DataPack.bNeedsConfirmation && !DataPack.bReceived)
+	{
+		DataPack.OnPacketRequest(Params.Header.DataVersion, Params.Header.DataOffset);
+	}
+};
+
+void AGameplayDebuggerCategoryReplicator::ClientDataPackPacket_Implementation(const FGameplayDebuggerDataPackRPCParams& Params)
+{
+	// Does the category exist on the client?
+	// using FindLastByPredicate since that's the only Predicate-based find that returns index.
+	const FName CategoryName = Params.CategoryName;
+	const int32 FoundIndex = Categories.FindLastByPredicate([CategoryName](const TSharedRef<FGameplayDebuggerCategory>& Item) { return (Item->GetCategoryName() == CategoryName);});	
+	if (FoundIndex == INDEX_NONE)
+	{
+		return;
+	}
+
+	// receive single packet
+	FGameplayDebuggerDataPack DataPacket;
+
+	DataPacket.Header = Params.Header;
+
+	const int32 PacketSize = FMath::Min(FGameplayDebuggerDataPack::PacketSize, DataPacket.Header.DataSize - DataPacket.Header.DataOffset);
+	if (PacketSize > 0)
+	{
+		DataPacket.Data.Append(Params.Data.GetData(), PacketSize);
+	}
+	
+	OnReceivedDataPackPacket(FoundIndex, Params.DataPackIdx, DataPacket);
+
+	const FGameplayDebuggerCategory& CategoryOb = Categories[FoundIndex].Get();
+	UE_LOG(LogGameplayDebugReplication, Verbose, TEXT("  >> replicate data pack[%d] progress:%.0f%%"), Params.DataPackIdx, CategoryOb.ReplicatedDataPacks[Params.DataPackIdx].GetProgress() * 100.0f);	
+}
+
+bool AGameplayDebuggerCategoryReplicator::GetViewPoint(FVector& OutViewLocation, FVector& OutViewDirection) const
+{
+	if (IsViewPointSet())
+	{
+		OutViewLocation = ViewLocation.GetValue();
+		OutViewDirection = ViewDirection.GetValue();
+		return true;
+	}
+
+	return false;
+}
+
+void AGameplayDebuggerCategoryReplicator::SetReplicatorOwner(APlayerController* InOwnerPC)
+{
+	if (!bIsEnabled)
+	{
+		// can't use bHasAuthority, BeginPlay was not called yet
+		UWorld* World = GetWorld();
+		if (FGameplayDebuggerUtils::IsAuthority(World))
+		{
+			APlayerController* OldOwner = OwnerPC;
+			OwnerPC = InOwnerPC;
+			NotifyDebuggerOwnerChange.Broadcast(this, OldOwner);
+
+			UE_VLOG_UELOG(this, LogGameplayDebugReplication, Log, TEXT("Set OWNER PC %s / %s / %s")
+				, *GetNameSafe(InOwnerPC), InOwnerPC ? *GetNameSafe(InOwnerPC->GetPawn()) : TEXT_EMPTY
+				, InOwnerPC ? *InOwnerPC->GetHumanReadableName() : TEXT_EMPTY);
+		}
+	}
+}
+
+void AGameplayDebuggerCategoryReplicator::SetEnabled(bool bEnable)
+{
+	UE_VLOG_UELOG(this, LogGameplayDebugReplication, Log, TEXT("SetEnabled %s"), TEXT_CONDITION(bEnable));
+
+	if (bHasAuthority)
+	{
+		bIsEnabled = bEnable;
+		bIsEnabledLocal = bEnable;
+		SetActorTickEnabled(bEnable);
+	}
+	else
+	{
+		ServerSetEnabled(bEnable);
+	}
+
+	MarkComponentsRenderStateDirty();
+	NotifyCategoriesToolState(bEnable);
+
+	// extensions will NOT work with simulate mode, they are meant to handle additional input
+	const bool bEnableExtensions = bEnable && !FGameplayDebuggerAddonBase::IsSimulateInEditor();
+	NotifyExtensionsToolState(bEnableExtensions);
+}
+
+void AGameplayDebuggerCategoryReplicator::SetDebugActor(AActor* Actor, bool bSelectInEditor)
+{
+	UE_VLOG_UELOG(this, LogGameplayDebugReplication, Log, TEXT("SetDebugActor %s"), *GetNameSafe(Actor));	
+	if (bHasAuthority)
+	{
+		if (DebugActor.Actor != Actor)
+		{
+			DebugActor.Actor = Actor;
+			DebugActor.ActorName = Actor ? Actor->GetFName() : NAME_None;
+			DebugActor.SyncCounter++;
+
+#if WITH_EDITOR
+			if (bSelectInEditor)
+			{
+				USelection* SelectedActors = GEditor ? GEditor->GetSelectedActors() : NULL;
+				if (SelectedActors && Actor)
+				{
+					SelectedActors->DeselectAll();
+					SelectedActors->Select(Actor);
+				}
+			}
+#endif // WITH_EDITOR
+		}
+	}
+	else
+	{
+		ServerSetDebugActor(Actor, bSelectInEditor);
+	}
+}
+
+void AGameplayDebuggerCategoryReplicator::SetViewPoint(const FVector& InViewLocation, const FVector& InViewDirection)
+{
+	ViewLocation = InViewLocation;
+	ViewDirection = InViewDirection;
+
+	if (!bHasAuthority)
+	{
+		ServerSetViewPoint(InViewLocation, InViewDirection);
+	}
+}
+
+void AGameplayDebuggerCategoryReplicator::ResetViewPoint()
+{
+	ViewLocation.Reset();
+	ViewDirection.Reset();
+
+	if (!bHasAuthority)
+	{
+		ServerResetViewPoint();
+	}
+}
+
+void AGameplayDebuggerCategoryReplicator::SetCategoryEnabled(int32 CategoryId, bool bEnable)
+{
+	if (bHasAuthority)
+	{
+		if (Categories.IsValidIndex(CategoryId))
+		{
+			UE_VLOG_UELOG(this, LogGameplayDebugReplication, Log, TEXT("SetCategoryEnabled[%d]:%d (%s)"), CategoryId, bEnable ? 1 : 0, *Categories[CategoryId]->GetCategoryName().ToString());
+			Categories[CategoryId]->bIsEnabled = bEnable;
+		}
+	}
+	else
+	{
+		ServerSetCategoryEnabled(CategoryId, bEnable);
+	}
+
+	MarkComponentsRenderStateDirty();
+}
+
+void AGameplayDebuggerCategoryReplicator::SendCategoryInputEvent(int32 CategoryId, int32 HandlerId)
+{
+	if (HandlerId >= 0 && Categories.IsValidIndex(CategoryId) &&
+		HandlerId < Categories[CategoryId]->GetNumInputHandlers())
+	{
+		// check enabled category only on local (instigating) side
+		if (!bIsLocal || IsCategoryEnabled(CategoryId))
+		{
+			FGameplayDebuggerInputHandler& InputHandler = Categories[CategoryId]->GetInputHandler(HandlerId);
+			if (InputHandler.Mode == EGameplayDebuggerInputMode::Local || bHasAuthority)
+			{
+				InputHandler.Delegate.ExecuteIfBound();
+			}
+			else
+			{
+				ServerSendCategoryInputEvent(CategoryId, HandlerId);
+			}
+		}
+	}
+}
+
+void AGameplayDebuggerCategoryReplicator::SendExtensionInputEvent(int32 ExtensionId, int32 HandlerId)
+{
+	if (HandlerId >= 0 && Extensions.IsValidIndex(ExtensionId) &&
+		HandlerId < Extensions[ExtensionId]->GetNumInputHandlers())
+	{
+		FGameplayDebuggerInputHandler& InputHandler = Extensions[ExtensionId]->GetInputHandler(HandlerId);
+		if (InputHandler.Mode == EGameplayDebuggerInputMode::Local || bHasAuthority)
+		{
+			InputHandler.Delegate.ExecuteIfBound();
+		}
+		else
+		{
+			ServerSendExtensionInputEvent(ExtensionId, HandlerId);
+		}
+	}
+}
+
+void AGameplayDebuggerCategoryReplicator::NotifyCategoriesToolState(bool bIsActive)
+{
+	for (int32 Idx = 0; Idx < Categories.Num(); Idx++)
+	{
+		FGameplayDebuggerCategory& CategoryOb = Categories[Idx].Get();
+		if (bIsActive)
+		{
+			CategoryOb.OnGameplayDebuggerActivated();
+		}
+		else
+		{
+			CategoryOb.OnGameplayDebuggerDeactivated();
+		}
+	}
+}
+
+void AGameplayDebuggerCategoryReplicator::NotifyExtensionsToolState(bool bIsActive)
+{
+	for (int32 Idx = 0; Idx < Extensions.Num(); Idx++)
+	{
+		FGameplayDebuggerExtension& ExtensionOb = Extensions[Idx].Get();
+		if (bIsActive)
+		{
+			ExtensionOb.OnGameplayDebuggerActivated();
+		}
+		else
+		{
+			ExtensionOb.OnGameplayDebuggerDeactivated();
+		}
+	}
+}
+
+bool AGameplayDebuggerCategoryReplicator::IsCategoryEnabled(int32 CategoryId) const
+{
+	return Categories.IsValidIndex(CategoryId) && Categories[CategoryId]->IsCategoryEnabled(); 
+}
+
diff --git a/Engine/Source/Runtime/GameplayDebugger/Private/GameplayDebuggerConfig.cpp b/Engine/Source/Runtime/GameplayDebugger/Private/GameplayDebuggerConfig.cpp
new file mode 100644
--- /dev/null
+++ b/Engine/Source/Runtime/GameplayDebugger/Private/GameplayDebuggerConfig.cpp
@@ -0,0 +1,302 @@
+// Copyright Epic Games, Inc. All Rights Reserved.
+
+#include "GameplayDebuggerConfig.h"
+#include "GameplayDebugger.h"
+#include "GameplayDebuggerAddonManager.h"
+#include "CoreGlobals.h"
+#include "Misc/ConfigCacheIni.h"
+
+#include UE_INLINE_GENERATED_CPP_BY_NAME(GameplayDebuggerConfig)
+
+
+UGameplayDebuggerConfig::UGameplayDebuggerConfig(const FObjectInitializer& ObjectInitializer) : Super(ObjectInitializer)
+{
+	ActivationKey = EKeys::Apostrophe;
+
+	CategoryRowNextKey = EKeys::Add;
+	CategoryRowPrevKey = EKeys::Subtract;
+
+	CategorySlot0 = EKeys::NumPadZero;
+	CategorySlot1 = EKeys::NumPadOne;
+	CategorySlot2 = EKeys::NumPadTwo;
+	CategorySlot3 = EKeys::NumPadThree;
+	CategorySlot4 = EKeys::NumPadFour;
+	CategorySlot5 = EKeys::NumPadFive;
+	CategorySlot6 = EKeys::NumPadSix;
+	CategorySlot7 = EKeys::NumPadSeven;
+	CategorySlot8 = EKeys::NumPadEight;
+	CategorySlot9 = EKeys::NumPadNine;
+
+	DebugCanvasPaddingLeft = 10.0f;
+	DebugCanvasPaddingRight = 10.0f;
+	DebugCanvasPaddingTop = 10.0f;
+	DebugCanvasPaddingBottom = 10.0f;
+}
+
+void UGameplayDebuggerConfig::UpdateCategoryConfig(const FName CategoryName, int32& SlotIdx, uint8& CategoryState)
+{
+	KnownCategoryNames.AddUnique(CategoryName);
+
+	int32 FoundIdx = INDEX_NONE;
+	for (int32 Idx = 0; Idx < Categories.Num(); Idx++)
+	{
+		if (*Categories[Idx].CategoryName == CategoryName)
+		{
+			FoundIdx = Idx;
+			break;
+		}
+	}
+
+	if (FoundIdx == INDEX_NONE)
+	{
+		FGameplayDebuggerCategoryConfig NewConfigData;
+		NewConfigData.CategoryName = CategoryName.ToString();
+
+		FoundIdx = Categories.Add(NewConfigData);
+	}
+
+	const EGameplayDebuggerCategoryState EnumState = (EGameplayDebuggerCategoryState)CategoryState;
+	if (Categories.IsValidIndex(FoundIdx))
+	{
+		FGameplayDebuggerCategoryConfig& ConfigData = Categories[FoundIdx];
+		if (ConfigData.bOverrideSlotIdx)
+		{
+			SlotIdx = ConfigData.SlotIdx;
+		}
+		else
+		{
+			ConfigData.SlotIdx = SlotIdx;
+		}
+		
+		const bool bDefaultActiveInGame = (EnumState == EGameplayDebuggerCategoryState::EnabledInGame) || (EnumState == EGameplayDebuggerCategoryState::EnabledInGameAndSimulate);
+		const bool bDefaultActiveInSimulate = (EnumState == EGameplayDebuggerCategoryState::EnabledInSimulate) || (EnumState == EGameplayDebuggerCategoryState::EnabledInGameAndSimulate);
+		const bool bDefaultHidden = (EnumState == EGameplayDebuggerCategoryState::Hidden);
+		
+		const bool bActiveInGame = (ConfigData.ActiveInGame == EGameplayDebuggerOverrideMode::UseDefault) ? bDefaultActiveInGame : (ConfigData.ActiveInGame == EGameplayDebuggerOverrideMode::Enable);
+		const bool bActiveInSimulate = (ConfigData.ActiveInSimulate == EGameplayDebuggerOverrideMode::UseDefault) ? bDefaultActiveInSimulate : (ConfigData.ActiveInSimulate == EGameplayDebuggerOverrideMode::Enable);
+		const bool bIsHidden = (ConfigData.Hidden == EGameplayDebuggerOverrideMode::UseDefault) ? bDefaultHidden : (ConfigData.Hidden == EGameplayDebuggerOverrideMode::Enable);
+
+		CategoryState = (uint8)(
+			bIsHidden ? EGameplayDebuggerCategoryState::Hidden :
+			bActiveInGame && bActiveInSimulate ? EGameplayDebuggerCategoryState::EnabledInGameAndSimulate :
+			bActiveInGame ? EGameplayDebuggerCategoryState::EnabledInGame :
+			bActiveInSimulate ? EGameplayDebuggerCategoryState::EnabledInSimulate :
+			EGameplayDebuggerCategoryState::Disabled);
+	}
+}
+
+void UGameplayDebuggerConfig::UpdateExtensionConfig(const FName ExtensionName, uint8& UseExtension)
+{
+	KnownExtensionNames.AddUnique(ExtensionName);
+
+	int32 FoundIdx = INDEX_NONE;
+	for (int32 Idx = 0; Idx < Extensions.Num(); Idx++)
+	{
+		if (*Extensions[Idx].ExtensionName == ExtensionName)
+		{
+			FoundIdx = Idx;
+			break;
+		}
+	}
+
+	if (FoundIdx == INDEX_NONE)
+	{
+		FGameplayDebuggerExtensionConfig NewConfigData;
+		NewConfigData.ExtensionName = ExtensionName.ToString();
+
+		FoundIdx = Extensions.Add(NewConfigData);
+	}
+
+	if (Extensions.IsValidIndex(FoundIdx))
+	{
+		FGameplayDebuggerExtensionConfig& ConfigData = Extensions[FoundIdx];
+
+		const bool bDefaultUseExtension = (UseExtension != 0);
+		const bool bUseExtension = (ConfigData.UseExtension == EGameplayDebuggerOverrideMode::UseDefault) ? bDefaultUseExtension : (ConfigData.UseExtension == EGameplayDebuggerOverrideMode::Enable);
+
+		UseExtension = bUseExtension ? 1 : 0;
+	}
+}
+
+void UGameplayDebuggerConfig::UpdateCategoryInputConfig(const FName CategoryName, const FName InputName, FName& KeyName, FGameplayDebuggerInputModifier& KeyModifier)
+{
+	int32 FoundCategoryIdx = INDEX_NONE;
+	for (int32 Idx = 0; Idx < Categories.Num(); Idx++)
+	{
+		if (*Categories[Idx].CategoryName == CategoryName)
+		{
+			FoundCategoryIdx = Idx;
+			break;
+		}
+	}
+
+	if (!Categories.IsValidIndex(FoundCategoryIdx))
+	{
+		return;
+	}
+
+	FGameplayDebuggerCategoryConfig& OuterConfigData = Categories[FoundCategoryIdx];
+	KnownCategoryInputNames.Add(CategoryName, InputName);
+
+	int32 FoundIdx = INDEX_NONE;
+	for (int32 Idx = 0; Idx < OuterConfigData.InputHandlers.Num(); Idx++)
+	{
+		if (*OuterConfigData.InputHandlers[Idx].ConfigName == InputName)
+		{
+			FGameplayDebuggerInputConfig& ConfigData = OuterConfigData.InputHandlers[Idx];
+			KeyName = ConfigData.Key.GetFName();
+			KeyModifier.bShift = ConfigData.bModShift;
+			KeyModifier.bCtrl = ConfigData.bModCtrl;
+			KeyModifier.bAlt = ConfigData.bModAlt;
+			KeyModifier.bCmd = ConfigData.bModCmd;
+
+			FoundIdx = Idx;
+			break;
+		}
+	}
+
+	if (FoundIdx == INDEX_NONE)
+	{
+		FGameplayDebuggerInputConfig NewConfigData;
+		NewConfigData.ConfigName = InputName.ToString();
+		NewConfigData.Key = FKey(KeyName);
+		NewConfigData.bModShift = KeyModifier.bShift;
+		NewConfigData.bModCtrl = KeyModifier.bCtrl;
+		NewConfigData.bModAlt = KeyModifier.bAlt;
+		NewConfigData.bModCmd = KeyModifier.bCmd;
+
+		OuterConfigData.InputHandlers.Add(NewConfigData);
+	}
+}
+
+void UGameplayDebuggerConfig::UpdateExtensionInputConfig(const FName ExtensionName, const FName InputName, FName& KeyName, FGameplayDebuggerInputModifier& KeyModifier)
+{
+	int32 FoundExtensionIdx = INDEX_NONE;
+	for (int32 Idx = 0; Idx < Extensions.Num(); Idx++)
+	{
+		if (*Extensions[Idx].ExtensionName == ExtensionName)
+		{
+			FoundExtensionIdx = Idx;
+			break;
+		}
+	}
+
+	if (!Extensions.IsValidIndex(FoundExtensionIdx))
+	{
+		return;
+	}
+
+	FGameplayDebuggerExtensionConfig& OuterConfigData = Extensions[FoundExtensionIdx];
+	KnownExtensionInputNames.Add(ExtensionName, InputName);
+
+	int32 FoundIdx = INDEX_NONE;
+	for (int32 Idx = 0; Idx < OuterConfigData.InputHandlers.Num(); Idx++)
+	{
+		if (*OuterConfigData.InputHandlers[Idx].ConfigName == InputName)
+		{
+			FGameplayDebuggerInputConfig& ConfigData = OuterConfigData.InputHandlers[Idx];
+			KeyName = ConfigData.Key.GetFName();
+			KeyModifier.bShift = ConfigData.bModShift;
+			KeyModifier.bCtrl = ConfigData.bModCtrl;
+			KeyModifier.bAlt = ConfigData.bModAlt;
+			KeyModifier.bCmd = ConfigData.bModCmd;
+
+			FoundIdx = Idx;
+			break;
+		}
+	}
+
+	if (FoundIdx == INDEX_NONE)
+	{
+		FGameplayDebuggerInputConfig NewConfigData;
+		NewConfigData.ConfigName = InputName.ToString();
+		NewConfigData.Key = FKey(KeyName);
+		NewConfigData.bModShift = KeyModifier.bShift;
+		NewConfigData.bModCtrl = KeyModifier.bCtrl;
+		NewConfigData.bModAlt = KeyModifier.bAlt;
+		NewConfigData.bModCmd = KeyModifier.bCmd;
+
+		OuterConfigData.InputHandlers.Add(NewConfigData);
+	}
+}
+
+void UGameplayDebuggerConfig::Serialize(FArchive& Ar)
+{
+	if (Ar.IsSaving())
+	{
+		RemoveUnknownConfigs();
+	}
+
+	Super::Serialize(Ar);
+}
+
+void UGameplayDebuggerConfig::RemoveUnknownConfigs()
+{
+	for (int32 Idx = Categories.Num() - 1; Idx >= 0; Idx--)
+	{
+		FGameplayDebuggerCategoryConfig& ConfigData = Categories[Idx];
+		const bool bValid = KnownCategoryNames.Contains(*ConfigData.CategoryName);
+		if (!bValid)
+		{
+			Categories.RemoveAt(Idx, 1, false);
+		}
+		else
+		{
+			for (int32 InputIdx = ConfigData.InputHandlers.Num() - 1; InputIdx >= 0; InputIdx--)
+			{
+				const FGameplayDebuggerInputConfig& InputConfigData = ConfigData.InputHandlers[InputIdx];
+				const FName* FoundDataPtr = KnownCategoryInputNames.FindPair(*ConfigData.CategoryName, *InputConfigData.ConfigName);
+				if (FoundDataPtr == nullptr)
+				{
+					ConfigData.InputHandlers.RemoveAt(InputIdx, 1, false);
+				}
+			}
+		}
+	}
+
+	for (int32 Idx = Extensions.Num() - 1; Idx >= 0; Idx--)
+	{
+		FGameplayDebuggerExtensionConfig& ConfigData = Extensions[Idx];
+		const bool bValid = KnownExtensionNames.Contains(*ConfigData.ExtensionName);
+		if (!bValid)
+		{
+			Extensions.RemoveAt(Idx, 1, false);
+		}
+		else
+		{
+			for (int32 InputIdx = ConfigData.InputHandlers.Num() - 1; InputIdx >= 0; InputIdx--)
+			{
+				const FGameplayDebuggerInputConfig& InputConfigData = ConfigData.InputHandlers[InputIdx];
+				const FName* FoundDataPtr = KnownExtensionInputNames.FindPair(*ConfigData.ExtensionName, *InputConfigData.ConfigName);
+				if (FoundDataPtr == nullptr)
+				{
+					ConfigData.InputHandlers.RemoveAt(InputIdx, 1, false);
+				}
+			}
+		}
+	}
+}
+
+#if WITH_EDITOR
+void UGameplayDebuggerConfig::PostEditChangeChainProperty(struct FPropertyChangedChainEvent& PropertyChangedEvent)
+{
+	Super::PostEditChangeChainProperty(PropertyChangedEvent);
+
+	FGameplayDebuggerAddonManager& AddonManager = FGameplayDebuggerAddonManager::GetCurrent();
+	AddonManager.UpdateFromConfig();
+}
+#endif
+
+//----------------------------------------------------------------------//
+//  UGameplayDebuggerUserSettings
+//----------------------------------------------------------------------//
+void UGameplayDebuggerUserSettings::SetFontSize(const int32 InFontSize)
+{
+	UGameplayDebuggerUserSettings* Setting = GetMutableDefault<UGameplayDebuggerUserSettings>();
+	Setting->FontSize = FMath::Clamp(InFontSize, 1, 100);
+
+	const FString ConfigFileName = Setting->GetProjectUserConfigFilename();
+	GConfig->SetInt(TEXT("/Script/GameplayDebugger.GameplayDebuggerUserSettings"), TEXT("FontSize"), Setting->FontSize, *ConfigFileName);
+	GConfig->Flush(false);
+}
+
diff --git a/Engine/Source/Runtime/GameplayDebugger/Private/GameplayDebuggerExtension.cpp b/Engine/Source/Runtime/GameplayDebugger/Private/GameplayDebuggerExtension.cpp
new file mode 100644
--- /dev/null
+++ b/Engine/Source/Runtime/GameplayDebugger/Private/GameplayDebuggerExtension.cpp
@@ -0,0 +1,50 @@
+// Copyright Epic Games, Inc. All Rights Reserved.
+
+#include "GameplayDebuggerExtension.h"
+#include "GameplayDebuggerCategoryReplicator.h"
+
+void FGameplayDebuggerExtension::OnGameplayDebuggerActivated()
+{
+	const bool bIsLocal = IsLocal();
+	if (bIsLocal)
+	{
+		OnActivated();
+	}
+}
+
+void FGameplayDebuggerExtension::OnGameplayDebuggerDeactivated()
+{
+	const bool bIsLocal = IsLocal();
+	if (bIsLocal)
+	{
+		OnDeactivated();
+	}
+}
+
+void FGameplayDebuggerExtension::OnActivated()
+{
+	// empty in base class
+}
+
+void FGameplayDebuggerExtension::OnDeactivated()
+{
+	// empty in base class
+}
+
+FString FGameplayDebuggerExtension::GetDescription() const
+{
+	// empty in base class
+	return FString();
+}
+
+APlayerController* FGameplayDebuggerExtension::GetPlayerController() const
+{
+	AGameplayDebuggerCategoryReplicator* Replicator = GetReplicator();
+	return Replicator ? Replicator->GetReplicationOwner() : nullptr;
+}
+
+bool FGameplayDebuggerExtension::IsLocal() const
+{
+	AGameplayDebuggerCategoryReplicator* Replicator = GetReplicator();
+	return Replicator ? Replicator->IsLocal() : true;
+}
diff --git a/Engine/Source/Runtime/GameplayDebugger/Private/GameplayDebuggerExtension_HUD.cpp b/Engine/Source/Runtime/GameplayDebugger/Private/GameplayDebuggerExtension_HUD.cpp
new file mode 100644
--- /dev/null
+++ b/Engine/Source/Runtime/GameplayDebugger/Private/GameplayDebuggerExtension_HUD.cpp
@@ -0,0 +1,116 @@
+// Copyright Epic Games, Inc. All Rights Reserved.
+
+#include "GameplayDebuggerExtension_HUD.h"
+#include "InputCoreTypes.h"
+#include "Engine/Engine.h"
+#include "GameFramework/PlayerController.h"
+#include "EngineGlobals.h"
+#include "GameFramework/HUD.h"
+
+FGameplayDebuggerExtension_HUD::FGameplayDebuggerExtension_HUD()
+{
+	bWantsHUDEnabled = false;
+	bIsGameHUDEnabled = false;
+	bAreDebugMessagesEnabled = false;
+	bPrevDebugMessagesEnabled = GEngine && GEngine->bEnableOnScreenDebugMessages;
+	bIsCachedDescriptionValid = false;
+
+	const FGameplayDebuggerInputHandlerConfig HUDKeyConfig(TEXT("ToggleHUD"), EKeys::Tilde.GetFName(), FGameplayDebuggerInputModifier::Ctrl);
+	const FGameplayDebuggerInputHandlerConfig MessagesKeyConfig(TEXT("ToggleMessages"), EKeys::Tab.GetFName(), FGameplayDebuggerInputModifier::Ctrl);
+
+	const bool bHasHUDBinding = BindKeyPress(HUDKeyConfig, this, &FGameplayDebuggerExtension_HUD::ToggleGameHUD);
+	HudBindingIdx = bHasHUDBinding ? (GetNumInputHandlers() - 1) : INDEX_NONE;
+
+	const bool bHasMessagesBinding = BindKeyPress(MessagesKeyConfig, this, &FGameplayDebuggerExtension_HUD::ToggleDebugMessages);
+	MessagesBindingIdx = bHasMessagesBinding ? (GetNumInputHandlers() - 1) : INDEX_NONE;
+}
+
+TSharedRef<FGameplayDebuggerExtension> FGameplayDebuggerExtension_HUD::MakeInstance()
+{
+	return MakeShareable(new FGameplayDebuggerExtension_HUD());
+}
+
+void FGameplayDebuggerExtension_HUD::OnActivated()
+{
+	SetGameHUDEnabled(bWantsHUDEnabled);
+	SetDebugMessagesEnabled(false);
+}
+
+void FGameplayDebuggerExtension_HUD::OnDeactivated()
+{
+	SetGameHUDEnabled(true);
+	SetDebugMessagesEnabled(bPrevDebugMessagesEnabled);
+}
+
+FString FGameplayDebuggerExtension_HUD::GetDescription() const
+{
+	if (!bIsCachedDescriptionValid)
+	{
+		FString Description;
+
+		if (HudBindingIdx != INDEX_NONE)
+		{
+			Description = FString::Printf(TEXT("{%s}%s:{%s}HUD"),
+				*FGameplayDebuggerCanvasStrings::ColorNameInput,
+				*GetInputHandlerDescription(HudBindingIdx),
+				bIsGameHUDEnabled ? *FGameplayDebuggerCanvasStrings::ColorNameEnabled : *FGameplayDebuggerCanvasStrings::ColorNameDisabled);
+		}
+
+		if (MessagesBindingIdx != INDEX_NONE)
+		{
+			if (Description.Len() > 0)
+			{
+				Description += FGameplayDebuggerCanvasStrings::SeparatorSpace;
+			}
+
+			Description += FString::Printf(TEXT("{%s}%s:{%s}DebugMessages"),
+				*FGameplayDebuggerCanvasStrings::ColorNameInput,
+				*GetInputHandlerDescription(MessagesBindingIdx),
+				bAreDebugMessagesEnabled ? *FGameplayDebuggerCanvasStrings::ColorNameEnabled : *FGameplayDebuggerCanvasStrings::ColorNameDisabled);
+		}
+
+		CachedDescription = Description;
+		bIsCachedDescriptionValid = true;
+	}
+
+	return CachedDescription;
+}
+
+void FGameplayDebuggerExtension_HUD::SetGameHUDEnabled(bool bEnable)
+{
+	APlayerController* OwnerPC = GetPlayerController();
+	AHUD* GameHUD = OwnerPC ? OwnerPC->GetHUD() : nullptr;
+	if (GameHUD)
+	{
+		// ShowHUD is actually a toggle and is a UFUNCTION(exec).
+		// We make sure the flag is !bEnable before calling the function so we don't end up desynced.
+		// Also, we don't simply set the flag to bEnable because we want to execute functions that overrides ShowHUD.
+		GameHUD->bShowHUD = !bEnable;
+		GameHUD->ShowHUD();
+	}
+
+	bIsGameHUDEnabled = bEnable;
+	bIsCachedDescriptionValid = false;
+}
+
+void FGameplayDebuggerExtension_HUD::SetDebugMessagesEnabled(bool bEnable)
+{
+	if (GEngine)
+	{
+		GEngine->bEnableOnScreenDebugMessages = bEnable;
+	}
+
+	bAreDebugMessagesEnabled = bEnable;
+	bIsCachedDescriptionValid = false;
+}
+
+void FGameplayDebuggerExtension_HUD::ToggleGameHUD()
+{
+	const bool bEnable = bWantsHUDEnabled = !bIsGameHUDEnabled;
+	SetGameHUDEnabled(bEnable);
+}
+
+void FGameplayDebuggerExtension_HUD::ToggleDebugMessages()
+{
+	SetDebugMessagesEnabled(!bAreDebugMessagesEnabled);
+}
diff --git a/Engine/Source/Runtime/GameplayDebugger/Private/GameplayDebuggerExtension_Spectator.cpp b/Engine/Source/Runtime/GameplayDebugger/Private/GameplayDebuggerExtension_Spectator.cpp
new file mode 100644
--- /dev/null
+++ b/Engine/Source/Runtime/GameplayDebugger/Private/GameplayDebuggerExtension_Spectator.cpp
@@ -0,0 +1,95 @@
+// Copyright Epic Games, Inc. All Rights Reserved.
+
+#include "GameplayDebuggerExtension_Spectator.h"
+#include "InputCoreTypes.h"
+#include "Engine/World.h"
+#include "GameplayDebuggerPlayerManager.h"
+#include "Engine/DebugCameraController.h"
+#include "Engine/Player.h"
+#include "GameFramework/HUD.h"
+#include "GameFramework/WorldSettings.h"
+
+FGameplayDebuggerExtension_Spectator::FGameplayDebuggerExtension_Spectator()
+{
+	const FGameplayDebuggerInputHandlerConfig KeyConfig(TEXT("Toggle"), EKeys::Tab.GetFName());
+	bHasInputBinding = BindKeyPress(KeyConfig, this, &FGameplayDebuggerExtension_Spectator::ToggleSpectatorMode);
+}
+
+TSharedRef<FGameplayDebuggerExtension> FGameplayDebuggerExtension_Spectator::MakeInstance()
+{
+	return MakeShareable(new FGameplayDebuggerExtension_Spectator());
+}
+
+void FGameplayDebuggerExtension_Spectator::OnDeactivated()
+{
+	if (SpectatorController.IsValid())
+	{
+		ToggleSpectatorMode();
+	}
+}
+
+FString FGameplayDebuggerExtension_Spectator::GetDescription() const
+{
+	if (!bIsCachedDescriptionValid)
+	{
+		CachedDescription = !bHasInputBinding ? FString() :
+			FString::Printf(TEXT("{%s}%s:{%s}Spectator"),
+				*FGameplayDebuggerCanvasStrings::ColorNameInput,
+				*GetInputHandlerDescription(0),
+				SpectatorController.IsValid() ? *FGameplayDebuggerCanvasStrings::ColorNameEnabled : *FGameplayDebuggerCanvasStrings::ColorNameDisabled);
+
+		bIsCachedDescriptionValid = true;
+	}
+
+	return CachedDescription;
+}
+
+void FGameplayDebuggerExtension_Spectator::ToggleSpectatorMode()
+{
+	APlayerController* OwnerPC = GetPlayerController();
+	AGameplayDebuggerPlayerManager& PlayerManager = AGameplayDebuggerPlayerManager::GetCurrent(OwnerPC->GetWorld());
+	UInputComponent* DebuggerInput = PlayerManager.GetInputComponent(*OwnerPC);
+
+	ADebugCameraController* SpectatorControllerOb = SpectatorController.Get();
+	if (SpectatorControllerOb == nullptr)
+	{
+		FActorSpawnParameters SpawnInfo;
+		SpawnInfo.SpawnCollisionHandlingOverride = ESpawnActorCollisionHandlingMethod::AlwaysSpawn;
+		SpawnInfo.Owner = OwnerPC->GetWorldSettings();
+		SpawnInfo.Instigator = OwnerPC->GetInstigator();
+		SpectatorControllerOb = OwnerPC->GetWorld()->SpawnActor<ADebugCameraController>(SpawnInfo);
+
+		if (SpectatorControllerOb)
+		{
+			OwnerPC->PopInputComponent(DebuggerInput);
+			SpectatorControllerOb->OnActivate(OwnerPC);
+			OwnerPC->Player->SwitchController(SpectatorControllerOb);
+			SpectatorControllerOb->PushInputComponent(DebuggerInput);
+
+			SpectatorControllerOb->ChangeState(NAME_Default);
+			SpectatorControllerOb->ChangeState(NAME_Spectating);
+			
+			if (SpectatorControllerOb->MyHUD)
+			{
+				SpectatorControllerOb->MyHUD->bShowHUD = false;
+			}
+
+			SpectatorController = SpectatorControllerOb;
+		}
+	}
+	else
+	{
+		SpectatorControllerOb->PopInputComponent(DebuggerInput);
+		if (SpectatorControllerOb->OriginalPlayer)
+		{
+			SpectatorControllerOb->OriginalPlayer->SwitchController(OwnerPC);
+		}
+		SpectatorControllerOb->OnDeactivate(OwnerPC);
+		OwnerPC->PushInputComponent(DebuggerInput);
+
+		SpectatorControllerOb->GetWorld()->DestroyActor(SpectatorControllerOb, false, false);
+		SpectatorController = nullptr;
+	}
+
+	bIsCachedDescriptionValid = false;
+}
diff --git a/Engine/Source/Runtime/GameplayDebugger/Private/GameplayDebuggerLocalController.cpp b/Engine/Source/Runtime/GameplayDebugger/Private/GameplayDebuggerLocalController.cpp
new file mode 100644
--- /dev/null
+++ b/Engine/Source/Runtime/GameplayDebugger/Private/GameplayDebuggerLocalController.cpp
@@ -0,0 +1,1100 @@
+// Copyright Epic Games, Inc. All Rights Reserved.
+
+#include "GameplayDebuggerLocalController.h"
+#include "Engine/Engine.h"
+#include "InputCoreTypes.h"
+#include "Framework/Commands/InputChord.h"
+#include "Components/InputComponent.h"
+#include "Misc/App.h"
+#include "SceneView.h"
+#include "TimerManager.h"
+#include "GameFramework/Pawn.h"
+#include "GameFramework/Controller.h"
+#include "GameFramework/SpectatorPawn.h"
+#include "GameplayDebuggerTypes.h"
+#include "GameplayDebuggerCategoryReplicator.h"
+#include "GameplayDebuggerPlayerManager.h"
+#include "GameplayDebuggerAddonBase.h"
+#include "GameplayDebuggerCategory.h"
+#include "GameplayDebuggerAddonManager.h"
+#include "GameplayDebuggerExtension.h"
+#include "GameplayDebuggerConfig.h"
+#include "GameplayDebuggerModule.h"
+#include "Debug/DebugDrawService.h"
+#include "Engine/Selection.h"
+#include "CanvasItem.h"
+#include "Engine/Canvas.h"
+#include "Engine/DebugCameraController.h"
+#include "GameFramework/PlayerInput.h"
+#include "EngineUtils.h"
+#include "HAL/IConsoleManager.h"
+
+#include UE_INLINE_GENERATED_CPP_BY_NAME(GameplayDebuggerLocalController)
+
+#if WITH_EDITOR
+#include "Editor.h"
+#endif // WITH_EDITOR
+
+UGameplayDebuggerLocalController::UGameplayDebuggerLocalController(const FObjectInitializer& ObjectInitializer) : Super(ObjectInitializer)
+{
+	bSimulateMode = false;
+	bNeedsCleanup = false;
+	bIsSelectingActor = false;
+	bIsLocallyEnabled = false;
+	bPrevLocallyEnabled = false;
+	bEnableTextShadow = false;
+	bPrevScreenMessagesEnabled = false;
+	ActiveRowIdx = 0;
+#if WITH_EDITOR
+	bActivateOnPIEEnd = false;
+#endif // WITH_EDITOR
+	if (HasAnyFlags(RF_ClassDefaultObject) == false)
+	{
+		HUDFont = NewObject<UFont>(this, TEXT("HUDFont"), RF_NoFlags, GEngine->GetSmallFont());
+		HUDFont->LegacyFontSize = UGameplayDebuggerUserSettings::GetFontSize(); //FGameplayDebuggerTweakables::FontSize;
+	}
+}
+
+void UGameplayDebuggerLocalController::Initialize(AGameplayDebuggerCategoryReplicator& Replicator, AGameplayDebuggerPlayerManager& Manager)
+{
+	CachedReplicator = &Replicator;
+	CachedPlayerManager = &Manager;
+	bSimulateMode = FGameplayDebuggerAddonBase::IsSimulateInEditor() || Replicator.IsEditorWorldReplicator();
+
+	UDebugDrawService::Register(bSimulateMode ? TEXT("DebugAI") : TEXT("Game"), FDebugDrawDelegate::CreateUObject(this, &UGameplayDebuggerLocalController::OnDebugDraw));
+
+#if WITH_EDITOR
+	if (bSimulateMode)
+	{
+		FGameplayDebuggerModule::OnLocalControllerInitialized.Broadcast();		
+	}
+
+	if (GIsEditor)
+	{
+		USelection::SelectionChangedEvent.AddUObject(this, &UGameplayDebuggerLocalController::OnSelectionChanged);
+		USelection::SelectObjectEvent.AddUObject(this, &UGameplayDebuggerLocalController::OnSelectedObject);
+
+		if (Replicator.IsEditorWorldReplicator())
+		{
+			// bind to PIE start and end notifies to hide before pie and re-enable if need be when pie's done
+			FEditorDelegates::BeginPIE.AddUObject(this, &UGameplayDebuggerLocalController::OnBeginPIE);
+			FEditorDelegates::EndPIE.AddUObject(this, &UGameplayDebuggerLocalController::OnEndPIE);
+		}
+	}
+#endif
+
+	FGameplayDebuggerAddonManager& AddonManager = FGameplayDebuggerAddonManager::GetCurrent();
+	AddonManager.OnCategoriesChanged.AddUObject(this, &UGameplayDebuggerLocalController::OnCategoriesChanged);
+	OnCategoriesChanged();
+
+	const UGameplayDebuggerConfig* SettingsCDO = UGameplayDebuggerConfig::StaticClass()->GetDefaultObject<UGameplayDebuggerConfig>();
+	const FKey NumpadKeys[] = { EKeys::NumPadZero, EKeys::NumPadOne, EKeys::NumPadTwo, EKeys::NumPadThree, EKeys::NumPadFour,
+		EKeys::NumPadFive, EKeys::NumPadSix, EKeys::NumPadSeven, EKeys::NumPadEight, EKeys::NumPadNine };
+	const FKey CategorySlots[] = { SettingsCDO->CategorySlot0, SettingsCDO->CategorySlot1, SettingsCDO->CategorySlot2, SettingsCDO->CategorySlot3, SettingsCDO->CategorySlot4,
+		SettingsCDO->CategorySlot5, SettingsCDO->CategorySlot6, SettingsCDO->CategorySlot7, SettingsCDO->CategorySlot8, SettingsCDO->CategorySlot9 };
+
+	bool bIsNumpadOnly = true;
+	for (int32 Idx = 0; Idx < UE_ARRAY_COUNT(CategorySlots); Idx++)
+	{
+		bool bHasPattern = false;
+		for (int32 PatternIdx = 0; PatternIdx < UE_ARRAY_COUNT(NumpadKeys); PatternIdx++)
+		{
+			if (CategorySlots[Idx] == NumpadKeys[PatternIdx])
+			{
+				bHasPattern = true;
+				break;
+			}
+		}
+
+		if (!bHasPattern)
+		{
+			bIsNumpadOnly = false;
+			break;
+		}
+	}
+
+	ActivationKeyDesc = GetKeyDescriptionLong(SettingsCDO->ActivationKey);
+	RowUpKeyDesc = GetKeyDescriptionShort(SettingsCDO->CategoryRowPrevKey);
+	RowDownKeyDesc = GetKeyDescriptionShort(SettingsCDO->CategoryRowNextKey);
+	CategoryKeysDesc = bIsNumpadOnly ? TEXT("{yellow}Numpad{white}") : TEXT("highlighted keys");
+
+	PaddingLeft = SettingsCDO->DebugCanvasPaddingLeft;
+	PaddingRight = SettingsCDO->DebugCanvasPaddingRight;
+	PaddingTop = SettingsCDO->DebugCanvasPaddingTop;
+	PaddingBottom = SettingsCDO->DebugCanvasPaddingBottom;
+
+	bEnableTextShadow = SettingsCDO->bDebugCanvasEnableTextShadow;
+
+	bNeedsCleanup = true;
+}
+
+void UGameplayDebuggerLocalController::Cleanup()
+{
+#if WITH_EDITOR
+	USelection::SelectionChangedEvent.RemoveAll(this);
+	USelection::SelectObjectEvent.RemoveAll(this);
+
+	if (bSimulateMode)
+	{
+		FGameplayDebuggerModule::OnLocalControllerUninitialized.Broadcast();
+	}
+#endif // WITH_EDITOR
+
+	bNeedsCleanup = false;
+}
+
+void UGameplayDebuggerLocalController::BeginDestroy()
+{
+	Super::BeginDestroy();
+	if (bNeedsCleanup)
+	{
+		Cleanup();
+	}
+}
+
+void UGameplayDebuggerLocalController::OnDebugDraw(class UCanvas* Canvas, class APlayerController* PC)
+{
+	if (CachedReplicator && CachedReplicator->IsEnabled())
+	{
+		FGameplayDebuggerCanvasContext CanvasContext(Canvas, HUDFont);
+		CanvasContext.CursorX = CanvasContext.DefaultX = PaddingLeft;
+		CanvasContext.CursorY = CanvasContext.DefaultY = PaddingTop;
+
+		CanvasContext.FontRenderInfo.bEnableShadow = bEnableTextShadow;
+
+		CanvasContext.PlayerController = CachedReplicator->GetReplicationOwner();
+		CanvasContext.World = CachedReplicator->GetReplicationOwner() ? CachedReplicator->GetReplicationOwner()->GetWorld() : CachedReplicator->GetWorld();
+
+		DrawHeader(CanvasContext);
+
+		if (DataPackMap.Num() != NumCategories)
+		{
+			RebuildDataPackMap();
+		}
+
+		if (Canvas->SceneView->ViewActor == nullptr)
+		{
+			CachedReplicator->SetViewPoint(Canvas->SceneView->ViewLocation, Canvas->SceneView->ViewRotation.Vector());
+		}
+		else if (CachedReplicator->IsViewPointSet())
+		{
+			CachedReplicator->ResetViewPoint();
+		}
+
+		const bool bHasDebugActor = CachedReplicator->HasDebugActor();
+		for (int32 Idx = 0; Idx < NumCategories; Idx++)
+		{
+			TSharedRef<FGameplayDebuggerCategory> Category = CachedReplicator->GetCategory(Idx);
+			if (Category->ShouldDrawCategory(bHasDebugActor))
+			{
+				if (Category->IsCategoryHeaderVisible())
+				{
+					DrawCategoryHeader(Idx, Category, CanvasContext);
+				}
+
+				Category->DrawCategory(CachedReplicator->GetReplicationOwner(), CanvasContext);
+			}
+		}
+	}
+}
+
+extern RENDERCORE_API FTexture* GWhiteTexture;
+
+void UGameplayDebuggerLocalController::DrawHeader(FGameplayDebuggerCanvasContext& CanvasContext)
+{
+	const int32 NumRows = (NumCategorySlots + (NumCategoriesPerRow-1)) / NumCategoriesPerRow;
+	const float LineHeight = CanvasContext.GetLineHeight();
+	const int32 NumExtensions = bSimulateMode ? 0 : CachedReplicator->GetNumExtensions();
+	const int32 NumExtensionRows = (NumExtensions > 0) ? 1 : 0;
+	const float CanvasSizeX = CanvasContext.Canvas->SizeX - PaddingLeft - PaddingRight;
+	const float UsePaddingTop = PaddingTop + (bSimulateMode ? 30.0f : 0);
+	
+	const float BackgroundPadding = 5.0f;
+	const float BackgroundPaddingBothSides = BackgroundPadding * 2.0f;
+
+	if (NumRows > 1)
+	{
+		FCanvasTileItem TileItemUpper(FVector2D(0, 0), GWhiteTexture, FVector2D(CanvasSizeX + BackgroundPaddingBothSides, (LineHeight * (ActiveRowIdx + NumExtensionRows + 1)) + BackgroundPadding), FLinearColor(0, 0, 0, 0.2f));
+		FCanvasTileItem ActiveRowTileItem(FVector2D(0, 0), GWhiteTexture, FVector2D(CanvasSizeX + BackgroundPaddingBothSides, LineHeight), FLinearColor(0, 0.5f, 0, 0.3f));
+		FCanvasTileItem TileItemLower(FVector2D(0, 0), GWhiteTexture, FVector2D(CanvasSizeX + BackgroundPaddingBothSides, LineHeight * ((NumRows - ActiveRowIdx - 1)) + BackgroundPadding), FLinearColor(0, 0, 0, 0.2f));
+
+		TileItemUpper.BlendMode = SE_BLEND_Translucent;
+		ActiveRowTileItem.BlendMode = SE_BLEND_Translucent;
+		TileItemLower.BlendMode = SE_BLEND_Translucent;
+
+		CanvasContext.DrawItem(TileItemUpper, PaddingLeft - BackgroundPadding, UsePaddingTop - BackgroundPadding);
+		CanvasContext.DrawItem(ActiveRowTileItem, PaddingLeft - BackgroundPadding, UsePaddingTop - BackgroundPadding + TileItemUpper.Size.Y);
+		CanvasContext.DrawItem(TileItemLower, PaddingLeft - BackgroundPadding, UsePaddingTop - BackgroundPadding + TileItemUpper.Size.Y + ActiveRowTileItem.Size.Y);
+	}
+	else
+	{
+		FCanvasTileItem TileItem(FVector2D(0, 0), GWhiteTexture, FVector2D(CanvasSizeX, LineHeight * (NumRows + NumExtensionRows + 1)) + BackgroundPaddingBothSides, FLinearColor(0, 0, 0, 0.2f));
+		TileItem.BlendMode = SE_BLEND_Translucent;
+		CanvasContext.DrawItem(TileItem, PaddingLeft - BackgroundPadding, UsePaddingTop - BackgroundPadding);
+	}
+
+	CanvasContext.CursorY = UsePaddingTop;
+	if (bSimulateMode)
+	{
+		CanvasContext.Printf(TEXT("Clear {yellow}DebugAI{white} show flag to close, use %s to toggle categories."), *CategoryKeysDesc);
+
+		// reactivate editor mode when this is being drawn = show flag is set
+#if WITH_EDITOR
+	FGameplayDebuggerModule::OnDebuggerEdModeActivation.Broadcast();
+#endif // WITH_EDITOR
+	}
+	else
+	{
+		CanvasContext.Printf(TEXT("Tap {yellow}%s{white} to close, use %s to toggle categories."), *ActivationKeyDesc, *CategoryKeysDesc);
+	}
+
+	const FString DebugActorDesc = FString::Printf(TEXT("Debug actor: {cyan}%s"), *CachedReplicator->GetDebugActorName().ToString());
+	float DebugActorSizeX = 0.0f, DebugActorSizeY = 0.0f;
+	CanvasContext.MeasureString(DebugActorDesc, DebugActorSizeX, DebugActorSizeY);
+	CanvasContext.PrintAt(CanvasContext.Canvas->SizeX - PaddingRight - DebugActorSizeX, UsePaddingTop, DebugActorDesc);
+
+	const FString VLogDesc = FString::Printf(TEXT("VLog: {cyan}%s"), CachedReplicator->GetVisLogSyncData().DeviceIDs.Len() > 0
+			? *CachedReplicator->GetVisLogSyncData().DeviceIDs
+			: TEXT("not recording to file"));
+		float VLogSizeX = 0.0f, VLogSizeY = 0.0f;
+		CanvasContext.MeasureString(VLogDesc, VLogSizeX, VLogSizeY);
+		CanvasContext.PrintAt(CanvasContext.Canvas->SizeX - PaddingRight - VLogSizeX, UsePaddingTop + LineHeight, VLogDesc);
+
+	const FString TimestampDesc = FString::Printf(TEXT("Time: %.2fs"), CachedReplicator->GetWorld()->GetTimeSeconds());
+	float TimestampSizeX = 0.0f, TimestampSizeY = 0.0f;
+	CanvasContext.MeasureString(TimestampDesc, TimestampSizeX, TimestampSizeY);
+	CanvasContext.PrintAt((CanvasSizeX - TimestampSizeX) * 0.5f, UsePaddingTop, TimestampDesc);
+
+	if (NumRows > 1)
+	{
+		const FString ChangeRowDesc = FString::Printf(TEXT("Prev row: {yellow}%s\n{white}Next row: {yellow}%s"), *RowUpKeyDesc, *RowDownKeyDesc);
+		float RowDescSizeX = 0.0f, RowDescSizeY = 0.0f;
+		CanvasContext.MeasureString(ChangeRowDesc, RowDescSizeX, RowDescSizeY);
+		CanvasContext.PrintAt(CanvasContext.Canvas->SizeX - PaddingRight - RowDescSizeX, UsePaddingTop + LineHeight * (NumExtensionRows + 1), ChangeRowDesc);
+	}
+
+	if (NumExtensionRows)
+	{
+		FString ExtensionRowDesc;
+		for (int32 ExtensionIdx = 0; ExtensionIdx < NumExtensions; ExtensionIdx++)
+		{
+			TSharedRef<FGameplayDebuggerExtension> Extension = CachedReplicator->GetExtension(ExtensionIdx);
+			FString ExtensionDesc = Extension->GetDescription();
+			ExtensionDesc.ReplaceInline(TEXT("\n"), TEXT(""));
+
+			if (ExtensionDesc.Len())
+			{
+				if (ExtensionRowDesc.Len())
+				{
+					ExtensionRowDesc += FGameplayDebuggerCanvasStrings::SeparatorSpace;
+				}
+
+				ExtensionRowDesc += ExtensionDesc;
+			}
+		}
+
+		CanvasContext.Print(ExtensionRowDesc);
+	}
+
+	for (int32 RowIdx = 0; RowIdx < NumRows; RowIdx++)
+	{
+		FString CategoryRowDesc;
+		for (int32 Idx = 0; Idx < NumCategoriesPerRow; Idx++)
+		{
+			const int32 CategorySlotIdx = (RowIdx * NumCategoriesPerRow) + Idx;
+			if (SlotCategoryIds.IsValidIndex(CategorySlotIdx) && 
+				SlotNames.IsValidIndex(CategorySlotIdx) &&
+				SlotCategoryIds[CategorySlotIdx].Num())
+			{
+				TSharedRef<FGameplayDebuggerCategory> Category0 = CachedReplicator->GetCategory(SlotCategoryIds[CategorySlotIdx][0]);
+				const bool bIsEnabled = Category0->IsCategoryEnabled();
+				const FString CategoryColorName = (RowIdx == ActiveRowIdx) && (NumRows > 1) ?
+					(bIsEnabled ? *FGameplayDebuggerCanvasStrings::ColorNameEnabledActiveRow : *FGameplayDebuggerCanvasStrings::ColorNameDisabledActiveRow) :
+					(bIsEnabled ? *FGameplayDebuggerCanvasStrings::ColorNameEnabled : *FGameplayDebuggerCanvasStrings::ColorNameDisabled);
+
+				const FString CategoryDesc = (RowIdx == ActiveRowIdx) ?
+					FString::Printf(TEXT("%s{%s}%d:{%s}%s"),
+						Idx ? *FGameplayDebuggerCanvasStrings::SeparatorSpace : TEXT(""),
+						*FGameplayDebuggerCanvasStrings::ColorNameInput,
+						Idx,
+						*CategoryColorName,
+						*SlotNames[CategorySlotIdx]) :
+					FString::Printf(TEXT("%s{%s}%s"),
+						Idx ? *FGameplayDebuggerCanvasStrings::Separator : TEXT(""),
+						*CategoryColorName,
+						*SlotNames[CategorySlotIdx]);
+
+				CategoryRowDesc += CategoryDesc;
+			}
+		}
+
+		CanvasContext.Print(CategoryRowDesc);
+	}
+
+	CanvasContext.DefaultY = CanvasContext.CursorY + LineHeight;
+}
+
+void UGameplayDebuggerLocalController::DrawCategoryHeader(int32 CategoryId, TSharedRef<FGameplayDebuggerCategory> Category, FGameplayDebuggerCanvasContext& CanvasContext)
+{
+	FString DataPackDesc;
+	
+	if (DataPackMap.IsValidIndex(CategoryId) &&
+		!Category->IsCategoryAuth() &&
+		!Category->ShouldDrawReplicationStatus() &&
+		Category->GetNumDataPacks() > 0)
+	{
+		// collect brief data pack status, detailed info is displayed only when ShouldDrawReplicationStatus is true
+		const int32 CurrentSyncCounter = CachedReplicator->GetDebugActorCounter();
+
+		DataPackDesc = TEXT("{white} ver[");
+		bool bIsPrevOutdated = false;
+		bool bAddSeparator = false;
+
+		for (int32 Idx = 0; Idx < DataPackMap[CategoryId].Num(); Idx++)
+		{
+			TSharedRef<FGameplayDebuggerCategory> MappedCategory = CachedReplicator->GetCategory(DataPackMap[CategoryId][Idx]);
+			for (int32 DataPackIdx = 0; DataPackIdx < MappedCategory->GetNumDataPacks(); DataPackIdx++)
+			{
+				FGameplayDebuggerDataPack::FHeader DataHeader = MappedCategory->GetDataPackHeaderCopy(DataPackIdx);
+				const bool bIsOutdated = (DataHeader.SyncCounter != CurrentSyncCounter);
+
+				if (bAddSeparator)
+				{
+					DataPackDesc += TEXT(';');
+				}
+
+				if (bIsOutdated != bIsPrevOutdated)
+				{
+					DataPackDesc += bIsOutdated ? TEXT("{red}") : TEXT("{white}");
+					bIsPrevOutdated = bIsOutdated;
+				}
+
+				DataPackDesc += TTypeToString<int16>::ToString(DataHeader.DataVersion);
+				bAddSeparator = true;
+			}
+		}
+
+		if (bIsPrevOutdated)
+		{
+			DataPackDesc += TEXT("{white}");
+		}
+
+		DataPackDesc += TEXT(']');
+	}
+
+	CanvasContext.MoveToNewLine();
+	CanvasContext.Printf(FColor::Green, TEXT("[CATEGORY: %s]%s"), *Category->GetCategoryName().ToString(), *DataPackDesc);
+}
+
+void UGameplayDebuggerLocalController::OnSelectLocalPlayer()
+{
+	APlayerController* OwnerPC = CachedReplicator ? CachedReplicator->GetReplicationOwner() : nullptr;
+	// Normal game. Spectator pawns aren't considered a valid local player to debug
+	if (OwnerPC && OwnerPC->Player)
+	{
+		if (AActor* LocalPlayerActor = OwnerPC->GetPawn())
+		{
+			CachedReplicator->SetDebugActor(LocalPlayerActor, true);
+			CachedReplicator->CollectCategoryData(/*bForce=*/true);
+		}
+	}
+}
+
+void UGameplayDebuggerLocalController::BindInput(UInputComponent& InputComponent)
+{
+	TSet<FName> NewBindings;
+
+	const UGameplayDebuggerConfig* SettingsCDO = UGameplayDebuggerConfig::StaticClass()->GetDefaultObject<UGameplayDebuggerConfig>();
+	if (!bSimulateMode)
+	{
+		InputComponent.BindKey(SettingsCDO->ActivationKey, IE_Pressed, this, &UGameplayDebuggerLocalController::OnActivationPressed);
+		InputComponent.BindKey(SettingsCDO->ActivationKey, IE_Released, this, &UGameplayDebuggerLocalController::OnActivationReleased);
+		NewBindings.Add(SettingsCDO->ActivationKey.GetFName());
+	}
+
+	if (bIsLocallyEnabled || bSimulateMode)
+	{
+		InputComponent.BindKey(SettingsCDO->CategorySlot0, IE_Pressed, this, &UGameplayDebuggerLocalController::OnCategory0Pressed);
+		InputComponent.BindKey(SettingsCDO->CategorySlot1, IE_Pressed, this, &UGameplayDebuggerLocalController::OnCategory1Pressed);
+		InputComponent.BindKey(SettingsCDO->CategorySlot2, IE_Pressed, this, &UGameplayDebuggerLocalController::OnCategory2Pressed);
+		InputComponent.BindKey(SettingsCDO->CategorySlot3, IE_Pressed, this, &UGameplayDebuggerLocalController::OnCategory3Pressed);
+		InputComponent.BindKey(SettingsCDO->CategorySlot4, IE_Pressed, this, &UGameplayDebuggerLocalController::OnCategory4Pressed);
+		InputComponent.BindKey(SettingsCDO->CategorySlot5, IE_Pressed, this, &UGameplayDebuggerLocalController::OnCategory5Pressed);
+		InputComponent.BindKey(SettingsCDO->CategorySlot6, IE_Pressed, this, &UGameplayDebuggerLocalController::OnCategory6Pressed);
+		InputComponent.BindKey(SettingsCDO->CategorySlot7, IE_Pressed, this, &UGameplayDebuggerLocalController::OnCategory7Pressed);
+		InputComponent.BindKey(SettingsCDO->CategorySlot8, IE_Pressed, this, &UGameplayDebuggerLocalController::OnCategory8Pressed);
+		InputComponent.BindKey(SettingsCDO->CategorySlot9, IE_Pressed, this, &UGameplayDebuggerLocalController::OnCategory9Pressed);
+
+		InputComponent.BindKey(SettingsCDO->CategoryRowPrevKey, IE_Pressed, this, &UGameplayDebuggerLocalController::OnCategoryRowUpPressed);
+		InputComponent.BindKey(SettingsCDO->CategoryRowNextKey, IE_Pressed, this, &UGameplayDebuggerLocalController::OnCategoryRowDownPressed);
+
+		NewBindings.Add(SettingsCDO->CategorySlot0.GetFName());
+		NewBindings.Add(SettingsCDO->CategorySlot1.GetFName());
+		NewBindings.Add(SettingsCDO->CategorySlot2.GetFName());
+		NewBindings.Add(SettingsCDO->CategorySlot3.GetFName());
+		NewBindings.Add(SettingsCDO->CategorySlot4.GetFName());
+		NewBindings.Add(SettingsCDO->CategorySlot5.GetFName());
+		NewBindings.Add(SettingsCDO->CategorySlot6.GetFName());
+		NewBindings.Add(SettingsCDO->CategorySlot7.GetFName());
+		NewBindings.Add(SettingsCDO->CategorySlot8.GetFName());
+		NewBindings.Add(SettingsCDO->CategorySlot9.GetFName());
+		NewBindings.Add(SettingsCDO->CategoryRowPrevKey.GetFName());
+		NewBindings.Add(SettingsCDO->CategoryRowNextKey.GetFName());
+
+		for (int32 Idx = 0; Idx < NumCategories; Idx++)
+		{
+			TSharedRef<FGameplayDebuggerCategory> Category = CachedReplicator->GetCategory(Idx);
+			const int32 NumInputHandlers = Category->GetNumInputHandlers();
+
+			for (int32 HandlerIdx = 0; HandlerIdx < NumInputHandlers; HandlerIdx++)
+			{
+				FGameplayDebuggerInputHandler& HandlerData = Category->GetInputHandler(HandlerIdx);
+				if (HandlerData.Modifier.bPressed || HandlerData.Modifier.bReleased)
+				{
+					FInputChord InputChord(FKey(HandlerData.KeyName), HandlerData.Modifier.bShift, HandlerData.Modifier.bCtrl, HandlerData.Modifier.bAlt, HandlerData.Modifier.bCmd);
+					FInputKeyBinding InputBinding(InputChord, HandlerData.Modifier.bPressed ? IE_Pressed : IE_Released);
+					InputBinding.KeyDelegate.GetDelegateForManualSet().BindUObject(this, &UGameplayDebuggerLocalController::OnCategoryBindingEvent, Idx, HandlerIdx);
+
+					InputComponent.KeyBindings.Add(InputBinding);
+					NewBindings.Add(HandlerData.KeyName);
+				}
+			}
+		}
+
+		const int32 NumExtentions = bSimulateMode ? 0 : CachedReplicator->GetNumExtensions();
+		for (int32 Idx = 0; Idx < NumExtentions; Idx++)
+		{
+			TSharedRef<FGameplayDebuggerExtension> Extension = CachedReplicator->GetExtension(Idx); //-V595
+			const int32 NumInputHandlers = Extension->GetNumInputHandlers();
+
+			for (int32 HandlerIdx = 0; HandlerIdx < NumInputHandlers; HandlerIdx++)
+			{
+				FGameplayDebuggerInputHandler& HandlerData = Extension->GetInputHandler(HandlerIdx);
+				if (HandlerData.Modifier.bPressed || HandlerData.Modifier.bReleased)
+				{
+					FInputChord InputChord(FKey(HandlerData.KeyName), HandlerData.Modifier.bShift, HandlerData.Modifier.bCtrl, HandlerData.Modifier.bAlt, HandlerData.Modifier.bCmd);
+					FInputKeyBinding InputBinding(InputChord, HandlerData.Modifier.bPressed ? IE_Pressed : IE_Released);
+					InputBinding.KeyDelegate.GetDelegateForManualSet().BindUObject(this, &UGameplayDebuggerLocalController::OnExtensionBindingEvent, Idx, HandlerIdx);
+
+					InputComponent.KeyBindings.Add(InputBinding);
+					NewBindings.Add(HandlerData.KeyName);
+				}
+			}
+		}
+	}
+
+	if (CachedReplicator && CachedReplicator->GetReplicationOwner() && CachedReplicator->GetReplicationOwner()->PlayerInput)
+	{
+		TSet<FName> RemovedMasks = UsedBindings.Difference(NewBindings);
+		TSet<FName> AddedMasks = NewBindings.Difference(UsedBindings);
+
+		UPlayerInput* Input = CachedReplicator->GetReplicationOwner()->PlayerInput;
+		for (int32 Idx = 0; Idx < Input->DebugExecBindings.Num(); Idx++)
+		{
+			FKeyBind& DebugBinding = Input->DebugExecBindings[Idx];
+			const bool bRemoveMask = RemovedMasks.Contains(DebugBinding.Key.GetFName());
+			const bool bAddMask = AddedMasks.Contains(DebugBinding.Key.GetFName());
+
+			if (bAddMask || bRemoveMask)
+			{
+				DebugBinding.bDisabled = bAddMask;
+			}
+		}
+
+		UsedBindings = NewBindings;
+	}
+}
+
+bool UGameplayDebuggerLocalController::IsKeyBound(const FName KeyName) const
+{
+	return UsedBindings.Contains(KeyName);
+}
+
+void UGameplayDebuggerLocalController::OnActivationPressed()
+{
+	if (CachedReplicator)
+	{
+		const float HoldTimeThr = 0.2f * (FApp::UseFixedTimeStep() ? (FApp::GetFixedDeltaTime() * 60.0f) : 1.0f);
+
+		CachedReplicator->GetWorldTimerManager().SetTimer(StartSelectingActorHandle, this, &UGameplayDebuggerLocalController::OnStartSelectingActor, HoldTimeThr);
+	}
+}
+
+void UGameplayDebuggerLocalController::OnActivationReleased()
+{
+	ToggleActivation();
+}
+
+void UGameplayDebuggerLocalController::ToggleActivation()
+{
+	if (CachedReplicator)
+	{
+		const UWorld* World = CachedReplicator->GetWorld();
+		if (!bIsSelectingActor || StartSelectingActorHandle.IsValid())
+		{
+			bPrevLocallyEnabled = bIsLocallyEnabled;
+			bIsLocallyEnabled = !CachedReplicator->IsEnabled();
+			CachedReplicator->SetEnabled(bIsLocallyEnabled);
+
+			if (bIsLocallyEnabled)
+			{
+				bPrevScreenMessagesEnabled = GAreScreenMessagesEnabled;
+				GAreScreenMessagesEnabled = false;
+				DebugActorCandidate = nullptr;
+				OnSelectActorTick();
+			}
+			else
+			{
+				// if Screen message are still disabled, restore previous state
+				if (!GAreScreenMessagesEnabled)
+				{
+					GAreScreenMessagesEnabled = bPrevScreenMessagesEnabled;
+				}
+			}
+		}
+
+		World->GetTimerManager().ClearTimer(StartSelectingActorHandle);
+		World->GetTimerManager().ClearTimer(SelectActorTickHandle);
+
+		CachedReplicator->MarkComponentsRenderStateDirty();
+	}
+
+	StartSelectingActorHandle.Invalidate();
+	SelectActorTickHandle.Invalidate();
+	bIsSelectingActor = false;
+
+	if (CachedReplicator && (bPrevLocallyEnabled != bIsLocallyEnabled))
+	{
+		CachedPlayerManager->RefreshInputBindings(*CachedReplicator);
+	}
+}
+
+void UGameplayDebuggerLocalController::OnCategory0Pressed()
+{
+	ToggleSlotState((ActiveRowIdx * NumCategoriesPerRow) + 0);
+}
+
+void UGameplayDebuggerLocalController::OnCategory1Pressed()
+{
+	ToggleSlotState((ActiveRowIdx * NumCategoriesPerRow) + 1);
+}
+
+void UGameplayDebuggerLocalController::OnCategory2Pressed()
+{
+	ToggleSlotState((ActiveRowIdx * NumCategoriesPerRow) + 2);
+}
+
+void UGameplayDebuggerLocalController::OnCategory3Pressed()
+{
+	ToggleSlotState((ActiveRowIdx * NumCategoriesPerRow) + 3);
+}
+
+void UGameplayDebuggerLocalController::OnCategory4Pressed()
+{
+	ToggleSlotState((ActiveRowIdx * NumCategoriesPerRow) + 4);
+}
+
+void UGameplayDebuggerLocalController::OnCategory5Pressed()
+{
+	ToggleSlotState((ActiveRowIdx * NumCategoriesPerRow) + 5);
+}
+
+void UGameplayDebuggerLocalController::OnCategory6Pressed()
+{
+	ToggleSlotState((ActiveRowIdx * NumCategoriesPerRow) + 6);
+}
+
+void UGameplayDebuggerLocalController::OnCategory7Pressed()
+{
+	ToggleSlotState((ActiveRowIdx * NumCategoriesPerRow) + 7);
+}
+
+void UGameplayDebuggerLocalController::OnCategory8Pressed()
+{
+	ToggleSlotState((ActiveRowIdx * NumCategoriesPerRow) + 8);
+}
+
+void UGameplayDebuggerLocalController::OnCategory9Pressed()
+{
+	ToggleSlotState((ActiveRowIdx * NumCategoriesPerRow) + 9);
+}
+
+void UGameplayDebuggerLocalController::OnCategoryRowUpPressed()
+{
+	const int32 NumRows = (NumCategorySlots + (NumCategoriesPerRow-1)) / NumCategoriesPerRow;
+	ActiveRowIdx = (NumRows > 1) ? ((ActiveRowIdx + NumRows - 1) % NumRows) : 0;
+}
+
+void UGameplayDebuggerLocalController::OnCategoryRowDownPressed()
+{
+	const int32 NumRows = (NumCategorySlots + (NumCategoriesPerRow-1)) / NumCategoriesPerRow;
+	ActiveRowIdx = (NumRows > 1) ? ((ActiveRowIdx + 1) % NumRows) : 0;
+}
+
+void UGameplayDebuggerLocalController::OnCategoryBindingEvent(int32 CategoryId, int32 HandlerId)
+{
+	if (CachedReplicator)
+	{
+		CachedReplicator->SendCategoryInputEvent(CategoryId, HandlerId);
+	}
+}
+
+void UGameplayDebuggerLocalController::OnExtensionBindingEvent(int32 ExtensionId, int32 HandlerId)
+{
+	if (CachedReplicator)
+	{
+		CachedReplicator->SendExtensionInputEvent(ExtensionId, HandlerId);
+	}
+}
+
+void UGameplayDebuggerLocalController::OnStartSelectingActor()
+{
+	StartSelectingActorHandle.Invalidate();
+	if (CachedReplicator)
+	{
+		if (!CachedReplicator->IsEnabled())
+		{
+			bPrevLocallyEnabled = bIsLocallyEnabled;
+			bIsLocallyEnabled = true;
+			CachedReplicator->SetEnabled(bIsLocallyEnabled);
+			bPrevScreenMessagesEnabled = GAreScreenMessagesEnabled;
+			GAreScreenMessagesEnabled = false;
+		}
+
+		bIsSelectingActor = true;
+		DebugActorCandidate = nullptr;
+
+		const bool bLooping = true;
+		CachedReplicator->GetWorldTimerManager().SetTimer(SelectActorTickHandle, this, &UGameplayDebuggerLocalController::OnSelectActorTick, 0.01f, bLooping);
+
+		OnSelectActorTick();
+	}
+}
+
+void UGameplayDebuggerLocalController::OnSelectActorTick()
+{
+	APlayerController* OwnerPC = CachedReplicator ? CachedReplicator->GetReplicationOwner() : nullptr;
+	if (OwnerPC)
+	{
+		FVector ViewLocation = FVector::ZeroVector;
+		FVector ViewDirection = FVector::ForwardVector;
+		if (!CachedReplicator->GetViewPoint(ViewLocation, ViewDirection))
+		{
+			AGameplayDebuggerPlayerManager::GetViewPoint(*OwnerPC, ViewLocation, ViewDirection);
+		}
+
+		const UGameplayDebuggerUserSettings* Settings = GetDefault<UGameplayDebuggerUserSettings>();
+		const float MaxScanDistance = Settings->MaxViewDistance;
+		const float MinViewDirDot = FMath::Cos(FMath::DegreesToRadians(Settings->MaxViewAngle));
+
+		AActor* BestCandidate = nullptr;
+		float BestScore = MinViewDirDot;
+
+		for (APawn* TestPawn : TActorRange<APawn>(OwnerPC->GetWorld()))
+		{
+			if (!TestPawn->IsHidden() && TestPawn->GetActorEnableCollision() &&
+				!TestPawn->IsA(ASpectatorPawn::StaticClass()) &&
+				TestPawn != OwnerPC->GetPawn())
+			{
+				FVector DirToPawn = (TestPawn->GetActorLocation() - ViewLocation);
+				float DistToPawn = DirToPawn.Size();
+				if (FMath::IsNearlyZero(DistToPawn))
+				{
+					DirToPawn = ViewDirection;
+					DistToPawn = 1.0f;
+				}
+				else
+				{
+					DirToPawn /= DistToPawn;
+				}
+
+				const float ViewDot = FVector::DotProduct(ViewDirection, DirToPawn);
+				if (DistToPawn < MaxScanDistance && ViewDot > BestScore)
+				{
+					BestScore = ViewDot;
+					BestCandidate = TestPawn;
+				}
+			}
+		}
+
+		// cache to avoid multiple RPC with the same actor
+		if (DebugActorCandidate != BestCandidate)
+		{
+			DebugActorCandidate = BestCandidate;
+			CachedReplicator->SetDebugActor(BestCandidate, true);
+		}
+	}
+}
+
+void UGameplayDebuggerLocalController::ToggleSlotState(int32 SlotIdx)
+{
+	if (CachedReplicator && SlotCategoryIds.IsValidIndex(SlotIdx) && SlotCategoryIds[SlotIdx].Num())
+	{
+		const bool bIsEnabled = CachedReplicator->IsCategoryEnabled(SlotCategoryIds[SlotIdx][0]);
+		for (int32 Idx = 0; Idx < SlotCategoryIds[SlotIdx].Num(); Idx++)
+		{
+			const int32 CategoryId = SlotCategoryIds[SlotIdx][Idx];
+			CachedReplicator->SetCategoryEnabled(CategoryId, !bIsEnabled);
+		}
+		// removed call to MarkComponentsRenderStateDirty since CachedReplicator->SetCategoryEnabled
+		// calls it already 
+	}
+}
+
+FString UGameplayDebuggerLocalController::GetKeyDescriptionShort(const FKey& KeyBind) const
+{
+	return FString::Printf(TEXT("[%s]"), *KeyBind.GetFName().ToString());
+}
+
+FString UGameplayDebuggerLocalController::GetKeyDescriptionLong(const FKey& KeyBind) const
+{
+	const FString KeyDisplay = KeyBind.GetDisplayName().ToString();
+	const FString KeyName = KeyBind.GetFName().ToString();
+	return (KeyDisplay == KeyName) ? FString::Printf(TEXT("[%s]"), *KeyDisplay) : FString::Printf(TEXT("%s [%s key])"), *KeyDisplay, *KeyName);
+}
+
+#if WITH_EDITOR
+void UGameplayDebuggerLocalController::OnSelectionChanged(UObject* Object)
+{
+	USelection* Selection = Cast<USelection>(Object);
+	if (Selection && CachedReplicator)
+	{
+		AActor* SelectedActor = nullptr;
+		for (int32 Idx = 0; Idx < Selection->Num(); Idx++)
+		{
+			SelectedActor = Cast<AActor>(Selection->GetSelectedObject(Idx));
+			if (SelectedActor)
+			{
+				break;
+			}
+		}
+
+		if (SelectedActor)
+		{
+			// Since this is used from a delegate we need to filter out objects not belonging to the same pie instance/world
+			if (SelectedActor && SelectedActor->GetWorld() != GetWorld())	
+			{
+				return;
+			}
+
+			CachedReplicator->SetDebugActor(SelectedActor, false);
+			CachedReplicator->CollectCategoryData(/*bForce=*/true);
+		}
+	}
+}
+
+void UGameplayDebuggerLocalController::OnSelectedObject(UObject* Object)
+{
+	// Since this is used from a delegate we need to filter out objects not belonging to the same pie instance/world
+	if (Object && Object->GetWorld() != GetWorld())	
+	{
+		return;
+	}
+
+	AController* SelectedController = Cast<AController>(Object);
+	APawn* SelectedPawn = SelectedController ? SelectedController->GetPawn() : Cast<APawn>(Object);
+	if (CachedReplicator && SelectedPawn && SelectedPawn->IsSelected())
+	{
+		CachedReplicator->SetDebugActor(SelectedPawn, false);
+		CachedReplicator->CollectCategoryData(/*bForce=*/true);
+	}
+}
+#endif
+
+void UGameplayDebuggerLocalController::OnCategoriesChanged()
+{
+	FGameplayDebuggerAddonManager& AddonManager = FGameplayDebuggerAddonManager::GetCurrent();
+
+	SlotNames.Reset();
+	SlotNames.Append(AddonManager.GetSlotNames());
+
+	// categories are already sorted using AddonManager.SlotMap, build Slot to Category Id map accordingly
+	const TArray< TArray<int32> >& SlotMap = AddonManager.GetSlotMap();
+	SlotCategoryIds.Reset();
+	SlotCategoryIds.AddDefaulted(SlotMap.Num());
+
+	int32 CategoryId = 0;
+	for (int32 SlotIdx = 0; SlotIdx < SlotMap.Num(); SlotIdx++)
+	{
+		for (int32 InnerIdx = 0; InnerIdx < SlotMap[SlotIdx].Num(); InnerIdx++)
+		{
+			SlotCategoryIds[SlotIdx].Add(CategoryId);
+			CategoryId++;
+		}
+	}
+
+	NumCategorySlots = SlotCategoryIds.Num();
+	NumCategories = AddonManager.GetNumVisibleCategories();
+
+	DataPackMap.Reset();
+}
+
+void UGameplayDebuggerLocalController::RebuildDataPackMap()
+{
+	DataPackMap.SetNum(NumCategories);
+	
+	// category: get all categories from slot and combine data pack data if category header is not displayed
+	for (int32 SlotIdx = 0; SlotIdx < NumCategorySlots; SlotIdx++)
+	{
+		TArray<int32> NoHeaderCategories;
+		int32 FirstVisibleCategoryId = INDEX_NONE;
+
+		for (int32 InnerIdx = 0; InnerIdx < SlotCategoryIds[SlotIdx].Num(); InnerIdx++)
+		{
+			const int32 CategoryId = SlotCategoryIds[SlotIdx][InnerIdx];
+			
+			TSharedRef<FGameplayDebuggerCategory> Category = CachedReplicator->GetCategory(CategoryId);
+			if (!Category->IsCategoryHeaderVisible())
+			{
+				NoHeaderCategories.Add(CategoryId);
+			}
+			else
+			{
+				DataPackMap[CategoryId].Add(CategoryId);
+				
+				if (FirstVisibleCategoryId == INDEX_NONE)
+				{
+					FirstVisibleCategoryId = CategoryId;
+				}
+			}
+		}
+
+		if ((FirstVisibleCategoryId != INDEX_NONE) && NoHeaderCategories.Num())
+		{
+			DataPackMap[FirstVisibleCategoryId].Append(NoHeaderCategories);
+		}
+	}
+}
+
+#if WITH_EDITOR
+void UGameplayDebuggerLocalController::OnBeginPIE(const bool bIsSimulating)
+{
+	bActivateOnPIEEnd = bIsLocallyEnabled;
+	if (bIsLocallyEnabled)
+	{
+		ToggleActivation();
+	}
+}
+
+void UGameplayDebuggerLocalController::OnEndPIE(const bool bIsSimulating)
+{
+	if (bActivateOnPIEEnd && !bIsLocallyEnabled)
+	{
+		ToggleActivation();
+	}
+}
+#endif // WITH_EDITOR
+
+//----------------------------------------------------------------------//
+// FGameplayDebuggerConsoleCommands 
+//----------------------------------------------------------------------//
+/**  Helper structure to declare/define console commands in the source file and to access UGameplayDebuggerLocalController protected members */
+struct FGameplayDebuggerConsoleCommands
+{
+private:
+	static UGameplayDebuggerLocalController* GetController(UWorld* InWorld)
+	{
+		UGameplayDebuggerLocalController* Controller = nullptr;
+
+		APlayerController* LocalPC = GEngine->GetFirstLocalPlayerController(InWorld);
+		if (LocalPC)
+		{
+			Controller = AGameplayDebuggerPlayerManager::GetCurrent(InWorld).GetLocalController(*LocalPC);
+		}
+#if WITH_EDITOR
+		else if (InWorld != nullptr && InWorld->IsGameWorld() == false)
+		{
+			Controller = AGameplayDebuggerPlayerManager::GetCurrent(InWorld).GetEditorController();
+		}
+#endif // WITH_EDITOR
+
+		UE_CLOG(Controller == nullptr, LogConsoleResponse, Error, TEXT("GameplayDebugger not available"));
+		return Controller;
+	}
+
+	static void ToggleGameplayDebugger(UWorld* InWorld)
+	{
+		if (UGameplayDebuggerLocalController* Controller = GetController(InWorld))
+		{
+			Controller->ToggleActivation();
+		}
+	}
+
+	static void SelectLocalPlayer(UWorld* InWorld)
+	{
+		if (UGameplayDebuggerLocalController* Controller = GetController(InWorld))
+		{
+			Controller->OnSelectLocalPlayer();
+		}
+	}
+
+	static void SelectPreviousRow(UWorld* InWorld)
+	{
+		if (UGameplayDebuggerLocalController* Controller = GetController(InWorld))
+		{
+			Controller->OnCategoryRowUpPressed();
+		}
+	}
+
+	static void SelectNextRow(UWorld* InWorld)
+	{
+		if (UGameplayDebuggerLocalController* Controller = GetController(InWorld))
+		{
+			Controller->OnCategoryRowDownPressed();
+		}
+	}
+
+	static void ToggleCategory(const TArray<FString>& Args, UWorld* InWorld)
+	{
+		UGameplayDebuggerLocalController* Controller = GetController(InWorld);
+		if (Controller == nullptr)
+		{
+			return;
+		}
+
+		if (Args.Num() != 1)
+		{
+			UE_LOG(LogConsoleResponse, Error, TEXT("Missing category index parameter. Usage: gdt.ToggleCategory <CategoryIdx>"));
+			return;
+		}
+
+		if (!Args[0].IsNumeric())
+		{
+			UE_LOG(LogConsoleResponse, Error, TEXT("Must provide numerical value as index. Usage: gdt.ToggleCategory <CategoryIdx>"));
+			return;
+		}
+		
+		const int32 SlotIdx = TCString<TCHAR>::Atoi(*Args[0]);
+		const int32 NumSlots = Controller->SlotCategoryIds.Num();
+		const int32 NumSlotsPerRow = UGameplayDebuggerLocalController::NumCategoriesPerRow;
+		const int32 NumRows = (NumSlots + (NumSlotsPerRow-1)) / NumSlotsPerRow;
+
+		const bool bIsLastRowActive = (Controller->ActiveRowIdx == NumRows-1);
+		const int32 NumSlotsOnActiveRow = bIsLastRowActive ? NumSlots - (NumSlotsPerRow * (NumRows-1)) : NumSlotsPerRow;
+		const int32 MaxSlotIdx = FMath::Max(0, FMath::Min(NumSlots, NumSlotsOnActiveRow)-1);
+		
+		if (!(Controller->SlotCategoryIds.IsValidIndex(SlotIdx) && SlotIdx <= MaxSlotIdx))
+		{
+			UE_LOG(LogConsoleResponse, Error, TEXT("Requires a category index in the active row [0..%d]. Usage: gdt.ToggleCategory CategoryIndex"), MaxSlotIdx);
+			return;
+		}
+		
+		Controller->ToggleSlotState((Controller->ActiveRowIdx * NumSlotsPerRow)+SlotIdx);
+	}
+
+	static void EnableCategoryName(const TArray<FString>& Args, UWorld* InWorld)
+	{
+		UGameplayDebuggerLocalController* Controller = GetController(InWorld);
+		if (Controller == nullptr || Controller->CachedReplicator == nullptr)
+		{
+			UE_CLOG(Controller != nullptr && Controller->CachedReplicator == nullptr, 
+				LogConsoleResponse, Error, TEXT("Failed due to CachedReplicator being Null"));
+			return;
+		}
+
+		if (Args.Num() < 1)
+		{
+			UE_LOG(LogConsoleResponse, Error, TEXT("Missing category name parameter. Usage: gdt.EnableCategory <CategoryNamePattern> [Enable]"));
+			return;
+		}
+
+		bool bEnable = (Args.Num() == 1);
+		if (Args.Num() > 1)
+		{
+			LexFromString(bEnable, *Args[1]);
+		}
+
+		for (int32 CategoryIndex = 0; CategoryIndex < Controller->CachedReplicator->GetNumCategories(); ++CategoryIndex)
+		{
+			const TSharedRef<FGameplayDebuggerCategory> Category = Controller->CachedReplicator->GetCategory(CategoryIndex);
+			const FString CategoryName = Category->GetCategoryName().ToString();
+			if (CategoryName.Find(Args[0], ESearchCase::IgnoreCase) != INDEX_NONE)
+			{
+				Controller->CachedReplicator->SetCategoryEnabled(CategoryIndex, bEnable);
+			}
+		}
+	}
+
+	static void SetFontSize(const TArray<FString>& Args, UWorld* InWorld)
+	{
+		if (Args.Num() != 1)
+		{
+			UE_LOG(LogConsoleResponse, Error, TEXT("Missing \'fontSize\' parameter. Usage: gdt.fontsize <fontSize>"));
+			return;
+		}
+
+		if (!Args[0].IsNumeric())
+		{
+			UE_LOG(LogConsoleResponse, Error, TEXT("Must provide numerical value as \'fontSize\'. Usage: gdt.fontsize <fontSize>"));
+			return;
+		}
+	
+		if (const UGameplayDebuggerLocalController* LocalController = GetController(InWorld))
+		{
+			UGameplayDebuggerUserSettings::SetFontSize(TCString<TCHAR>::Atoi(*Args[0]));
+			check(LocalController->HUDFont);
+			LocalController->HUDFont->LegacyFontSize = UGameplayDebuggerUserSettings::GetFontSize();
+		}
+	}
+
+	/** For legacy command: EnableGDT */
+	static FAutoConsoleCommandWithWorld EnableDebuggerCmd;
+
+	/** Various gameplay debugger commands: gdt.<command> */
+	static FAutoConsoleCommandWithWorld ToggleDebuggerCmd;
+	static FAutoConsoleCommandWithWorld SelectLocalPlayerCmd;
+	static FAutoConsoleCommandWithWorld SelectPreviousRowCmd;
+	static FAutoConsoleCommandWithWorld SelectNextRowCmd;
+	static FAutoConsoleCommandWithWorldAndArgs ToggleCategoryCmd;
+	static FAutoConsoleCommandWithWorldAndArgs EnableCategoryNameCmd;
+	static FAutoConsoleCommandWithWorldAndArgs SetFontSizeCmd;
+};
+
+FAutoConsoleCommandWithWorld FGameplayDebuggerConsoleCommands::EnableDebuggerCmd(
+	TEXT("EnableGDT"),
+	TEXT("Toggles Gameplay Debugger Tool"),
+	FConsoleCommandWithWorldDelegate::CreateStatic(&FGameplayDebuggerConsoleCommands::ToggleGameplayDebugger)
+);
+
+FAutoConsoleCommandWithWorld FGameplayDebuggerConsoleCommands::ToggleDebuggerCmd(
+	TEXT("gdt.Toggle"),
+	TEXT("Toggles Gameplay Debugger Tool"),
+	FConsoleCommandWithWorldDelegate::CreateStatic(&FGameplayDebuggerConsoleCommands::ToggleGameplayDebugger)
+);
+
+FAutoConsoleCommandWithWorld FGameplayDebuggerConsoleCommands::SelectLocalPlayerCmd(
+	TEXT("gdt.SelectLocalPlayer"),
+	TEXT("Selects the local player for debugging"),
+	FConsoleCommandWithWorldDelegate::CreateStatic(FGameplayDebuggerConsoleCommands::SelectLocalPlayer)
+);
+
+FAutoConsoleCommandWithWorld FGameplayDebuggerConsoleCommands::SelectPreviousRowCmd(
+	TEXT("gdt.SelectPreviousRow"),
+	TEXT("Selects previous row"),
+	FConsoleCommandWithWorldDelegate::CreateStatic(FGameplayDebuggerConsoleCommands::SelectPreviousRow)
+);
+
+FAutoConsoleCommandWithWorld FGameplayDebuggerConsoleCommands::SelectNextRowCmd(
+	TEXT("gdt.SelectNextRow"),
+	TEXT("Selects next row"),
+	FConsoleCommandWithWorldDelegate::CreateStatic(FGameplayDebuggerConsoleCommands::SelectNextRow)
+);
+
+FAutoConsoleCommandWithWorldAndArgs FGameplayDebuggerConsoleCommands::ToggleCategoryCmd(
+	TEXT("gdt.ToggleCategory"),
+	TEXT("Toggles specific category index"),
+	FConsoleCommandWithWorldAndArgsDelegate::CreateStatic(&FGameplayDebuggerConsoleCommands::ToggleCategory)
+);
+
+FAutoConsoleCommandWithWorldAndArgs FGameplayDebuggerConsoleCommands::EnableCategoryNameCmd(
+	TEXT("gdt.EnableCategoryName"),
+	TEXT("Enables/disables categories matching given substring. Use: gdt.EnableCategoryName <CategoryNamePart> [Enable]"),
+	FConsoleCommandWithWorldAndArgsDelegate::CreateStatic(&FGameplayDebuggerConsoleCommands::EnableCategoryName)
+);
+
+FAutoConsoleCommandWithWorldAndArgs FGameplayDebuggerConsoleCommands::SetFontSizeCmd(
+	TEXT("gdt.fontsize"),
+	TEXT("Configures gameplay debugger's font size. Usage: gdt.fontsize <fontSize> (default = 10)"),
+	FConsoleCommandWithWorldAndArgsDelegate::CreateStatic(&FGameplayDebuggerConsoleCommands::SetFontSize)
+);
+
diff --git a/Engine/Source/Runtime/GameplayDebugger/Private/GameplayDebuggerModule.cpp b/Engine/Source/Runtime/GameplayDebugger/Private/GameplayDebuggerModule.cpp
new file mode 100644
--- /dev/null
+++ b/Engine/Source/Runtime/GameplayDebugger/Private/GameplayDebuggerModule.cpp
@@ -0,0 +1,147 @@
+// Copyright Epic Games, Inc. All Rights Reserved.
+
+#include "GameplayDebuggerModule.h"
+
+#include "CoreMinimal.h"
+#include "Modules/ModuleManager.h"
+#include "Engine/World.h"
+#include "GameplayDebugger.h"
+#include "ISettingsModule.h"
+#include "GameplayDebuggerAddonManager.h"
+#include "GameplayDebuggerPlayerManager.h"
+#include "GameplayDebuggerConfig.h"
+
+#include "GameplayDebuggerExtension_Spectator.h"
+#include "GameplayDebuggerExtension_HUD.h"
+
+#if UE_WITH_IRIS
+#include "Iris/IrisConfig.h"
+#endif // UE_WITH_IRIS
+
+FOnLocalControllerInitialized FGameplayDebuggerModule::OnLocalControllerInitialized;
+FOnLocalControllerUninitialized FGameplayDebuggerModule::OnLocalControllerUninitialized;
+
+#if WITH_EDITOR
+FOnDebuggerEdMode FGameplayDebuggerModule::OnDebuggerEdModeActivation;
+#endif //WITH_EDITOR
+
+IMPLEMENT_MODULE(FGameplayDebuggerModule, GameplayDebugger)
+
+void FGameplayDebuggerModule::StartupModule()
+{
+	// This code will execute after your module is loaded into memory (but after global variables are initialized, of course.)
+	FWorldDelegates::OnPostWorldInitialization.AddRaw(this, &FGameplayDebuggerModule::OnWorldInitialized);
+
+	UGameplayDebuggerConfig* SettingsCDO = UGameplayDebuggerConfig::StaticClass()->GetDefaultObject<UGameplayDebuggerConfig>();
+	if (SettingsCDO)
+	{
+		ISettingsModule* SettingsModule = FModuleManager::GetModulePtr<ISettingsModule>("Settings");
+		if (SettingsModule)
+		{
+			SettingsModule->RegisterSettings("Project", "Engine", "GameplayDebugger",
+				NSLOCTEXT("GameplayDebuggerModule", "SettingsName", "Gameplay Debugger"),
+				NSLOCTEXT("GameplayDebuggerModule", "SettingsDescription", "Settings for the gameplay debugger tool."),
+				SettingsCDO);
+		}
+
+		AddonManager.RegisterExtension("GameHUD", FOnGetExtension::CreateStatic(&FGameplayDebuggerExtension_HUD::MakeInstance));
+		AddonManager.RegisterExtension("Spectator", FOnGetExtension::CreateStatic(&FGameplayDebuggerExtension_Spectator::MakeInstance));
+		AddonManager.NotifyExtensionsChanged();
+	}
+}
+
+void FGameplayDebuggerModule::ShutdownModule()
+{
+	// This function may be called during shutdown to clean up your module.  For modules that support dynamic reloading,
+	// we call this function before unloading the module.
+	FWorldDelegates::OnPostWorldInitialization.RemoveAll(this);
+
+	ISettingsModule* SettingsModule = FModuleManager::GetModulePtr<ISettingsModule>("Settings");
+	if (SettingsModule)
+	{
+		SettingsModule->UnregisterSettings("Project", "Engine", "GameplayDebugger");
+	}
+}
+
+void FGameplayDebuggerModule::RegisterCategory(FName CategoryName, IGameplayDebugger::FOnGetCategory MakeInstanceDelegate, EGameplayDebuggerCategoryState CategoryState, int32 SlotIdx)
+{
+	AddonManager.RegisterCategory(CategoryName, MakeInstanceDelegate, CategoryState, SlotIdx);
+}
+
+void FGameplayDebuggerModule::UnregisterCategory(FName CategoryName)
+{
+	AddonManager.UnregisterCategory(CategoryName);
+}
+
+void FGameplayDebuggerModule::NotifyCategoriesChanged()
+{
+	AddonManager.NotifyCategoriesChanged();
+}
+
+void FGameplayDebuggerModule::RegisterExtension(FName ExtensionName, IGameplayDebugger::FOnGetExtension MakeInstanceDelegate)
+{
+	AddonManager.RegisterExtension(ExtensionName, MakeInstanceDelegate);
+}
+
+void FGameplayDebuggerModule::UnregisterExtension(FName ExtensionName)
+{
+	AddonManager.UnregisterExtension(ExtensionName);
+}
+
+void FGameplayDebuggerModule::NotifyExtensionsChanged()
+{
+	AddonManager.NotifyExtensionsChanged();
+}
+
+FGameplayDebuggerAddonManager& FGameplayDebuggerAddonManager::GetCurrent()
+{
+	FGameplayDebuggerModule& Module = FModuleManager::LoadModuleChecked<FGameplayDebuggerModule>("GameplayDebugger");
+	return Module.AddonManager;
+}
+
+AGameplayDebuggerPlayerManager& AGameplayDebuggerPlayerManager::GetCurrent(UWorld* World)
+{
+	FGameplayDebuggerModule& Module = FModuleManager::LoadModuleChecked<FGameplayDebuggerModule>("GameplayDebugger");
+	return Module.GetPlayerManager(World);
+}
+
+AGameplayDebuggerPlayerManager& FGameplayDebuggerModule::GetPlayerManager(UWorld* World)
+{
+	const int32 PurgeInvalidWorldsSize = 5;
+	if (PlayerManagers.Num() > PurgeInvalidWorldsSize)
+	{
+		for (TMap<TWeakObjectPtr<UWorld>, TWeakObjectPtr<AGameplayDebuggerPlayerManager> >::TIterator It(PlayerManagers); It; ++It)
+		{
+			if (!It.Key().IsValid())
+			{
+				It.RemoveCurrent();
+			}
+			else if (!It.Value().IsValid())
+			{
+				It.RemoveCurrent();
+			}
+		}
+	}
+
+	TWeakObjectPtr<AGameplayDebuggerPlayerManager> Manager = PlayerManagers.FindRef(World);
+	AGameplayDebuggerPlayerManager* ManagerOb = Manager.Get();
+
+	if (ManagerOb == nullptr)
+	{
+		ManagerOb = World->SpawnActor<AGameplayDebuggerPlayerManager>();
+		PlayerManagers.Add(World, ManagerOb);
+		ManagerOb->Init();
+	}
+
+	check(ManagerOb);
+	return *ManagerOb;
+}
+
+void FGameplayDebuggerModule::OnWorldInitialized(UWorld* World, const UWorld::InitializationValues IVS)
+{
+	// make sure that world has valid player manager, create when it doesn't
+	if (World && (World->IsGameWorld() || World->IsEditorWorld()))
+	{
+		GetPlayerManager(World);
+	}
+}
\ No newline at end of file
diff --git a/Engine/Source/Runtime/GameplayDebugger/Private/GameplayDebuggerPlayerManager.cpp b/Engine/Source/Runtime/GameplayDebugger/Private/GameplayDebuggerPlayerManager.cpp
new file mode 100644
--- /dev/null
+++ b/Engine/Source/Runtime/GameplayDebugger/Private/GameplayDebuggerPlayerManager.cpp
@@ -0,0 +1,360 @@
+// Copyright Epic Games, Inc. All Rights Reserved.
+
+#include "GameplayDebuggerPlayerManager.h"
+#include "Engine/World.h"
+#include "Components/InputComponent.h"
+#include "GameplayDebuggerCategoryReplicator.h"
+#include "GameplayDebuggerLocalController.h"
+#include "Engine/DebugCameraController.h"
+#include "Engine/InputDelegateBinding.h"
+#include "GameplayDebuggerConfig.h"
+#include "UnrealEngine.h"
+#include "Engine/LocalPlayer.h"
+#include "Net/UnrealNetwork.h"
+#include "GameFramework/GameModeBase.h"
+
+#include UE_INLINE_GENERATED_CPP_BY_NAME(GameplayDebuggerPlayerManager)
+
+AGameplayDebuggerPlayerManager::AGameplayDebuggerPlayerManager(const FObjectInitializer& ObjectInitializer) : Super(ObjectInitializer)
+{
+	PrimaryActorTick.bCanEverTick = true;
+	PrimaryActorTick.bAllowTickOnDedicatedServer = true;
+	PrimaryActorTick.bTickEvenWhenPaused = true;
+	PrimaryActorTick.bStartWithTickEnabled = false;
+	PrimaryActorTick.TickInterval = 0.5f;
+
+#if WITH_EDITOR
+	SetIsTemporarilyHiddenInEditor(true);
+#endif
+
+#if WITH_EDITORONLY_DATA
+	bHiddenEdLevel = true;
+	bHiddenEdLayer = true;
+	bHiddenEd = true;
+	bEditable = false;
+#endif
+
+	bIsLocal = false;
+	bInitialized = false;
+	bEditorTimeTick = false;
+}
+
+void AGameplayDebuggerPlayerManager::PostInitProperties()
+{
+	Super::PostInitProperties();
+
+	if (HasAnyFlags(RF_ClassDefaultObject) == false)
+	{
+		if (InputComponent == nullptr)
+		{
+			// create an InputComponent object so that the level script actor can bind key events
+			InputComponent = NewObject<UInputComponent>(this, TEXT("GameplayDebug_Input"));
+			InputComponent->RegisterComponent();
+		}
+
+		if (UInputDelegateBinding::SupportsInputDelegate(GetClass()))
+		{
+			UInputDelegateBinding::BindInputDelegates(GetClass(), InputComponent);
+		}
+#if WITH_EDITORONLY_DATA 
+		const UWorld* World = GetWorld();
+		bEditorTimeTick = (World != nullptr) && (World->IsEditorWorld() == true) && (World->IsGameWorld() == false);
+#endif // WITH_EDITORONLY_DATA 
+	}
+}
+
+void AGameplayDebuggerPlayerManager::BeginPlay()
+{
+	Super::BeginPlay();
+
+	UWorld* World = GetWorld();
+	check(World);
+	const ENetMode NetMode = World->GetNetMode();
+	
+	bHasAuthority = FGameplayDebuggerUtils::IsAuthority(World);
+	bIsLocal = (NetMode != NM_DedicatedServer);
+	bInitialized = true;
+
+	if (bHasAuthority)
+	{
+		UpdateAuthReplicators();
+		SetActorTickEnabled(true);
+	}
+	
+	for (int32 Idx = 0; Idx < PendingRegistrations.Num(); Idx++)
+	{
+		RegisterReplicator(*PendingRegistrations[Idx]);
+	}
+
+	PendingRegistrations.Empty();
+
+	FNetworkReplayDelegates::OnScrubTeardown.AddUObject(this, &ThisClass::OnReplayScrubTeardown);
+	FGameModeEvents::GameModeLogoutEvent.AddUObject(this, &ThisClass::OnGameModeLogout);
+}
+
+void AGameplayDebuggerPlayerManager::EndPlay(const EEndPlayReason::Type Reason)
+{
+	Super::EndPlay(Reason);
+
+	for (int32 Idx = 0; Idx < PlayerData.Num(); Idx++)
+	{
+		FGameplayDebuggerPlayerData& TestData = PlayerData[Idx];
+		if (IsValid(TestData.Controller))
+		{
+			TestData.Controller->Cleanup();
+			TestData.Controller = nullptr;
+		}
+	}
+
+	FNetworkReplayDelegates::OnScrubTeardown.RemoveAll(this);
+	FGameModeEvents::GameModeLogoutEvent.RemoveAll(this);
+}
+
+void AGameplayDebuggerPlayerManager::Init()
+{
+#if WITH_EDITOR
+	UWorld* World = GetWorld();
+	if (World != nullptr && World->WorldType == EWorldType::Editor && (GetDefault<UGameplayDebuggerUserSettings>()->bEnableGameplayDebuggerInEditor))
+	{
+		bHasAuthority = true;
+		bIsLocal = true;
+		bInitialized = true;
+
+		AGameplayDebuggerCategoryReplicator* Replicator = World->SpawnActorDeferred<AGameplayDebuggerCategoryReplicator>(AGameplayDebuggerCategoryReplicator::StaticClass(), FTransform::Identity);
+		Replicator->SetReplicatorOwner(nullptr);
+		Replicator->FinishSpawning(FTransform::Identity, true);
+		SetActorTickEnabled(true);
+
+		EditorWorldData.Replicator = Replicator;
+
+		Replicator->InitForEditor();
+
+		EditorWorldData.Controller = NewObject<UGameplayDebuggerLocalController>(this, TEXT("GameplayDebug_Controller_Editor"));
+		EditorWorldData.Controller->Initialize(*Replicator, *this);
+		
+	}
+#endif // WITH_EDITOR
+}
+
+void AGameplayDebuggerPlayerManager::TickActor(float DeltaTime, enum ELevelTick TickType, FActorTickFunction& ThisTickFunction)
+{
+	Super::TickActor(DeltaTime, TickType, ThisTickFunction);
+	UpdateAuthReplicators();
+};
+
+void AGameplayDebuggerPlayerManager::UpdateAuthReplicators()
+{
+	UWorld* World = GetWorld();
+	for (int32 Idx = PlayerData.Num() - 1; Idx >= 0; Idx--)
+	{
+		FGameplayDebuggerPlayerData& TestData = PlayerData[Idx];
+
+		if (!IsValid(TestData.Replicator) || !IsValid(TestData.Replicator->GetReplicationOwner()))
+		{
+			if (IsValid(TestData.Replicator))
+			{
+				World->DestroyActor(TestData.Replicator);
+			}
+
+			if (IsValid(TestData.Controller))
+			{
+				TestData.Controller->Cleanup();
+			}
+
+			PlayerData.RemoveAt(Idx, 1, false);
+		}
+	}
+
+	for (FConstPlayerControllerIterator It = World->GetPlayerControllerIterator(); It; It++)
+	{
+		APlayerController* TestPC = It->Get();
+		if (TestPC && !TestPC->IsA<ADebugCameraController>())
+		{
+			const bool bNeedsReplicator = (GetReplicator(*TestPC) == nullptr);
+			if (bNeedsReplicator)
+			{
+				AGameplayDebuggerCategoryReplicator* Replicator = World->SpawnActorDeferred<AGameplayDebuggerCategoryReplicator>(AGameplayDebuggerCategoryReplicator::StaticClass(), FTransform::Identity);
+				Replicator->SetReplicatorOwner(TestPC);
+				Replicator->FinishSpawning(FTransform::Identity, true);
+			}
+		}
+	}
+
+	PrimaryActorTick.TickInterval = PlayerData.Num() ? 5.0f : 0.5f;
+}
+
+void AGameplayDebuggerPlayerManager::RegisterReplicator(AGameplayDebuggerCategoryReplicator& Replicator)
+{
+	if (!bInitialized)
+	{
+		PendingRegistrations.Add(&Replicator);
+		return;
+	}
+
+	// keep all player related objects together for easy access and GC
+	FGameplayDebuggerPlayerData NewData;
+	NewData.Replicator = &Replicator;
+	
+	if (bIsLocal)
+	{
+		APlayerController* OwnerPC = Replicator.GetReplicationOwner();
+
+		NewData.InputComponent = OwnerPC ? NewObject<UInputComponent>(OwnerPC, TEXT("GameplayDebug_Input")) : ToRawPtr(InputComponent);
+		check(NewData.InputComponent);
+		NewData.InputComponent->Priority = -1;
+		NewData.Controller = NewObject<UGameplayDebuggerLocalController>(OwnerPC ? OwnerPC : (AActor*)this, TEXT("GameplayDebug_Controller"));
+		NewData.Controller->Initialize(Replicator, *this);
+		NewData.Controller->BindInput(*NewData.InputComponent);
+
+		if (OwnerPC)
+		{
+			OwnerPC->PushInputComponent(NewData.InputComponent);
+		}
+	}
+	else
+	{
+		NewData.Controller = nullptr;
+		NewData.InputComponent = nullptr;
+	}
+
+	PlayerData.Add(NewData);
+}
+
+void AGameplayDebuggerPlayerManager::RefreshInputBindings(AGameplayDebuggerCategoryReplicator& Replicator)
+{
+	for (int32 Idx = 0; Idx < PlayerData.Num(); Idx++)
+	{
+		FGameplayDebuggerPlayerData& TestData = PlayerData[Idx];
+		if (TestData.Replicator == &Replicator)
+		{
+			TestData.InputComponent->ClearActionBindings();
+			TestData.InputComponent->ClearBindingValues();
+			TestData.InputComponent->KeyBindings.Empty();
+
+			TestData.Controller->BindInput(*TestData.InputComponent);
+		}
+	}
+}
+
+AGameplayDebuggerCategoryReplicator* AGameplayDebuggerPlayerManager::GetReplicator(const APlayerController& OwnerPC) const
+{
+	const FGameplayDebuggerPlayerData* DataPtr = GetPlayerData(OwnerPC);
+	return DataPtr ? DataPtr->Replicator : nullptr;
+}
+
+UInputComponent* AGameplayDebuggerPlayerManager::GetInputComponent(const APlayerController& OwnerPC) const
+{
+	const FGameplayDebuggerPlayerData* DataPtr = GetPlayerData(OwnerPC);
+	return DataPtr ? DataPtr->InputComponent : nullptr;
+}
+
+UGameplayDebuggerLocalController* AGameplayDebuggerPlayerManager::GetLocalController(const APlayerController& OwnerPC) const
+{
+	const FGameplayDebuggerPlayerData* DataPtr = GetPlayerData(OwnerPC);
+	return DataPtr ? DataPtr->Controller : nullptr;
+}
+
+const FGameplayDebuggerPlayerData* AGameplayDebuggerPlayerManager::GetPlayerData(const APlayerController& OwnerPC) const
+{
+	for (int32 Idx = 0; Idx < PlayerData.Num(); Idx++)
+	{
+		const FGameplayDebuggerPlayerData& TestData = PlayerData[Idx];
+		if (TestData.Replicator && TestData.Replicator->GetReplicationOwner() == &OwnerPC)
+		{
+			return &TestData;
+		}
+	}
+
+	return nullptr;
+}
+
+void AGameplayDebuggerPlayerManager::GetViewPoint(const APlayerController& OwnerPC, FVector& OutViewLocation, FVector& OutViewDirection)
+{
+	UWorld* World = OwnerPC.GetWorld();
+	FVector CameraLocation;
+	FRotator CameraRotation;
+	if (OwnerPC.Player)
+	{
+		// normal game
+		OwnerPC.GetPlayerViewPoint(CameraLocation, CameraRotation);
+	}
+	else
+	{
+		// spectator mode
+		for (FLocalPlayerIterator It(GEngine, World); It; ++It)
+		{
+			ADebugCameraController* SpectatorPC = Cast<ADebugCameraController>(It->PlayerController);
+			if (SpectatorPC)
+			{
+				SpectatorPC->GetPlayerViewPoint(CameraLocation, CameraRotation);
+				break;
+			}
+		}
+	}
+
+	OutViewLocation = CameraLocation;
+	OutViewDirection = CameraRotation.Vector();
+}
+
+// FTickableGameObject begin
+void AGameplayDebuggerPlayerManager::Tick(float DeltaTime)
+{
+#if WITH_EDITORONLY_DATA 
+	if (EditorWorldData.Replicator)
+	{
+		FActorTickFunction DummyTickFunction;
+		EditorWorldData.Replicator->TickActor(DeltaTime, ELevelTick::LEVELTICK_All, DummyTickFunction);
+	}
+#endif // WITH_EDITORONLY_DATA 
+}
+
+ETickableTickType AGameplayDebuggerPlayerManager::GetTickableTickType() const
+{
+	return
+#if WITH_EDITOR
+		IsTickable() ? ETickableTickType::Conditional : 
+#endif // WITH_EDITOR
+		ETickableTickType::Never;
+}
+
+TStatId AGameplayDebuggerPlayerManager::GetStatId() const
+{
+	RETURN_QUICK_DECLARE_CYCLE_STAT(AGameplayDebuggerPlayerManager, STATGROUP_Tickables);
+}
+// FTickableGameObject end
+
+void AGameplayDebuggerPlayerManager::OnGameModeLogout(AGameModeBase* GameMode, AController* Exiting)
+{
+	if (GameMode && GameMode->GetWorld() == GetWorld())
+	{
+		UWorld* World = GetWorld();
+		for (int32 Idx = PlayerData.Num() - 1; Idx >= 0; Idx--)
+		{
+			FGameplayDebuggerPlayerData& TestData = PlayerData[Idx];
+
+			if (IsValid(TestData.Replicator) && (TestData.Replicator->GetReplicationOwner() == Exiting))
+			{
+				if (IsValid(TestData.Replicator))
+				{
+					World->DestroyActor(TestData.Replicator);
+				}
+
+				if (IsValid(TestData.Controller))
+				{
+					TestData.Controller->Cleanup();
+				}
+
+				PlayerData.RemoveAt(Idx, 1, false);
+				break;
+			}
+		}
+	}
+}
+
+void AGameplayDebuggerPlayerManager::OnReplayScrubTeardown(UWorld* InWorld)
+{
+	if (GetWorld() == InWorld)
+	{
+		UpdateAuthReplicators();
+	}
+}
diff --git a/Engine/Source/Runtime/GameplayDebugger/Private/GameplayDebuggerRenderingComponent.cpp b/Engine/Source/Runtime/GameplayDebugger/Private/GameplayDebuggerRenderingComponent.cpp
new file mode 100644
--- /dev/null
+++ b/Engine/Source/Runtime/GameplayDebugger/Private/GameplayDebuggerRenderingComponent.cpp
@@ -0,0 +1,174 @@
+// Copyright Epic Games, Inc. All Rights Reserved.
+
+#include "GameplayDebuggerRenderingComponent.h"
+#include "GameplayDebuggerCategoryReplicator.h"
+#include "GameplayDebuggerCategory.h"
+
+#include UE_INLINE_GENERATED_CPP_BY_NAME(GameplayDebuggerRenderingComponent)
+
+//////////////////////////////////////////////////////////////////////////
+// FGameplayDebuggerCompositeSceneProxy
+
+class FGameplayDebuggerCompositeSceneProxy : public FDebugRenderSceneProxy
+{
+	friend class FGameplayDebuggerDebugDrawDelegateHelper;
+public:
+	FGameplayDebuggerCompositeSceneProxy(const UPrimitiveComponent* InComponent) : FDebugRenderSceneProxy(InComponent) { }
+
+	virtual ~FGameplayDebuggerCompositeSceneProxy() override
+	{
+		for (int32 Idx = 0; Idx < ChildProxies.Num(); Idx++)
+		{
+			delete ChildProxies[Idx];
+		}
+	}
+
+	virtual void DrawStaticElements(FStaticPrimitiveDrawInterface* PDI) override
+	{
+		for (int32 Idx = 0; Idx < ChildProxies.Num(); Idx++)
+		{
+			ChildProxies[Idx]->DrawStaticElements(PDI);
+		}
+	}
+
+	virtual void GetDynamicMeshElements(const TArray<const FSceneView*>& Views, const FSceneViewFamily& ViewFamily, uint32 VisibilityMap, FMeshElementCollector& Collector) const override
+	{
+		for (int32 Idx = 0; Idx < ChildProxies.Num(); Idx++)
+		{
+			ChildProxies[Idx]->GetDynamicMeshElements(Views, ViewFamily, VisibilityMap, Collector);
+		}
+	}
+
+	virtual FPrimitiveViewRelevance GetViewRelevance(const FSceneView* View) const override
+	{
+		FPrimitiveViewRelevance Result;
+		for (int32 Idx = 0; Idx < ChildProxies.Num(); Idx++)
+		{
+			Result |= ChildProxies[Idx]->GetViewRelevance(View);
+		}
+		return Result;
+	}
+
+	virtual uint32 GetMemoryFootprint(void) const override
+	{
+		return sizeof(*this) + GetAllocatedSizeInternal();
+	}
+
+	void AddChild(FDebugRenderSceneProxy* NewChild)
+	{
+		ChildProxies.AddUnique(NewChild);
+	}
+
+	void AddRange(TArray<FDebugRenderSceneProxy*> Children)
+	{
+		ChildProxies.Append(Children);
+	}
+
+private:
+	uint32 GetAllocatedSizeInternal(void) const
+	{
+		uint32 Size = FDebugRenderSceneProxy::GetAllocatedSize() + ChildProxies.GetAllocatedSize();
+		for (int32 Idx = 0; Idx < ChildProxies.Num(); Idx++)
+		{
+			Size += ChildProxies[Idx]->GetMemoryFootprint();
+		}
+
+		return Size;
+	}
+
+protected:
+	TArray<FDebugRenderSceneProxy*> ChildProxies;
+};
+
+void FGameplayDebuggerDebugDrawDelegateHelper::RegisterDebugDrawDelegateInternal()
+{
+	ensureMsgf(State != RegisteredState, TEXT("DrawDelegate is already Registered!"));
+	if (State == InitializedState)
+	{
+		for (int32 Idx = 0; Idx < DebugDrawDelegateHelpers.Num(); Idx++)
+		{
+			DebugDrawDelegateHelpers[Idx]->RequestRegisterDebugDrawDelegate(nullptr /*Deferred Context*/);
+		}
+		State = RegisteredState;
+	}
+}
+
+void FGameplayDebuggerDebugDrawDelegateHelper::UnregisterDebugDrawDelegate()
+{
+	ensureMsgf(State != InitializedState, TEXT("DrawDelegate is in an invalid State: %i !"), State);
+	if (State == RegisteredState)
+	{
+		for (int32 Idx = 0; Idx < DebugDrawDelegateHelpers.Num(); Idx++)
+		{
+			DebugDrawDelegateHelpers[Idx]->UnregisterDebugDrawDelegate();
+		}
+		State = InitializedState;
+	}
+}
+
+void FGameplayDebuggerDebugDrawDelegateHelper::Reset()
+{
+	for (int32 Idx = 0; Idx < DebugDrawDelegateHelpers.Num(); Idx++)
+	{
+		delete DebugDrawDelegateHelpers[Idx];
+	}
+	DebugDrawDelegateHelpers.Reset();
+}
+
+void FGameplayDebuggerDebugDrawDelegateHelper::AddDelegateHelper(FDebugDrawDelegateHelper* InDebugDrawDelegateHelper)
+{
+	check(InDebugDrawDelegateHelper);
+	DebugDrawDelegateHelpers.Add(InDebugDrawDelegateHelper);
+}
+
+//////////////////////////////////////////////////////////////////////////
+// UGameplayDebuggerRenderingComponent
+
+UGameplayDebuggerRenderingComponent::UGameplayDebuggerRenderingComponent(const FObjectInitializer& ObjInitializer) : Super(ObjInitializer)
+{
+}
+
+FDebugRenderSceneProxy* UGameplayDebuggerRenderingComponent::CreateDebugSceneProxy()
+{
+	GameplayDebuggerDebugDrawDelegateHelper.Reset();
+
+	FGameplayDebuggerCompositeSceneProxy* CompositeProxy = nullptr;
+
+	AGameplayDebuggerCategoryReplicator* OwnerReplicator = Cast<AGameplayDebuggerCategoryReplicator>(GetOwner());
+	if (OwnerReplicator && OwnerReplicator->IsEnabled())
+	{
+		TArray<FDebugRenderSceneProxy*> SceneProxies;
+		for (int32 Idx = 0; Idx < OwnerReplicator->GetNumCategories(); Idx++)
+		{
+			const TSharedRef<FGameplayDebuggerCategory> Category = OwnerReplicator->GetCategory(Idx);
+			if (Category->IsCategoryEnabled())
+			{
+				FDebugDrawDelegateHelper* CategoryDelegateHelper = nullptr;
+				FDebugRenderSceneProxy* CategorySceneProxy = Category->CreateDebugSceneProxy(this, CategoryDelegateHelper);
+				if (CategorySceneProxy)
+				{
+					SceneProxies.Add(CategorySceneProxy);
+				}
+
+				if (CategoryDelegateHelper)
+				{
+					GameplayDebuggerDebugDrawDelegateHelper.AddDelegateHelper(CategoryDelegateHelper);
+				}
+			}
+		}
+
+		if (SceneProxies.Num())
+		{
+			CompositeProxy = new FGameplayDebuggerCompositeSceneProxy(this);
+			CompositeProxy->AddRange(SceneProxies);
+		}
+	}
+
+	return CompositeProxy;
+}
+
+FBoxSphereBounds UGameplayDebuggerRenderingComponent::CalcBounds(const FTransform &LocalToWorld) const
+{
+	return FBoxSphereBounds(FBox::BuildAABB(FVector::ZeroVector, FVector(1000000.0f, 1000000.0f, 1000000.0f)));
+}
+
diff --git a/Engine/Source/Runtime/GameplayDebugger/Private/GameplayDebuggerTypes.cpp b/Engine/Source/Runtime/GameplayDebugger/Private/GameplayDebuggerTypes.cpp
new file mode 100644
--- /dev/null
+++ b/Engine/Source/Runtime/GameplayDebugger/Private/GameplayDebuggerTypes.cpp
@@ -0,0 +1,832 @@
+// Copyright Epic Games, Inc. All Rights Reserved.
+
+#include "GameplayDebuggerTypes.h"
+#include "InputCoreTypes.h"
+#include "SceneView.h"
+#include "Serialization/MemoryWriter.h"
+#include "Serialization/MemoryReader.h"
+#include "GameplayDebuggerConfig.h"
+#include "DrawDebugHelpers.h"
+#include "CanvasItem.h"
+#include "Engine/Canvas.h"
+#include "Engine/World.h"
+#include "GameFramework/PlayerController.h"
+
+#include UE_INLINE_GENERATED_CPP_BY_NAME(GameplayDebuggerTypes)
+
+DEFINE_LOG_CATEGORY(LogGameplayDebug);
+
+namespace FGameplayDebuggerUtils
+{
+	bool IsAuthority(UWorld* World)
+	{
+		return (World == nullptr) || (World->GetNetMode() != NM_Client) || World->IsPlayingReplay();
+	}
+}
+
+//////////////////////////////////////////////////////////////////////////
+// FGameplayDebuggerShape
+
+FGameplayDebuggerShape FGameplayDebuggerShape::MakePoint(const FVector& Location, const float Radius, const FColor& Color, const FString& Description)
+{
+	FGameplayDebuggerShape NewElement;
+	NewElement.ShapeData.Add(Location);
+	NewElement.ShapeData.Add(FVector(Radius, 0, 0));
+	NewElement.Color = Color;
+	NewElement.Description = Description;
+	NewElement.Type = EGameplayDebuggerShape::Point;
+
+	return NewElement;
+}
+
+FGameplayDebuggerShape FGameplayDebuggerShape::MakeSegment(const FVector& StartLocation, const FVector& EndLocation, const float Thickness, const FColor& Color, const FString& Description)
+{
+	FGameplayDebuggerShape NewElement;
+	NewElement.ShapeData.Add(StartLocation);
+	NewElement.ShapeData.Add(EndLocation);
+	NewElement.ShapeData.Add(FVector(Thickness, 0, 0));
+	NewElement.Color = Color;
+	NewElement.Description = Description;
+	NewElement.Type = EGameplayDebuggerShape::Segment;
+
+	return NewElement;
+}
+
+FGameplayDebuggerShape FGameplayDebuggerShape::MakeSegment(const FVector& StartLocation, const FVector& EndLocation, const FColor& Color, const FString& Description)
+{
+	return MakeSegment(StartLocation, EndLocation, 1.0f, Color, Description);
+}
+
+FGameplayDebuggerShape FGameplayDebuggerShape::MakeArrow(const FVector& StartLocation, const FVector& EndLocation, const float HeadSize, const float Thickness, const FColor& Color, const FString& Description)
+{
+	FGameplayDebuggerShape NewElement;
+	NewElement.ShapeData.Add(StartLocation);
+	NewElement.ShapeData.Add(EndLocation);
+	NewElement.ShapeData.Add(FVector(Thickness, HeadSize, 0));
+	NewElement.Color = Color;
+	NewElement.Description = Description;
+	NewElement.Type = EGameplayDebuggerShape::Arrow;
+
+	return NewElement;
+}
+
+FGameplayDebuggerShape FGameplayDebuggerShape::MakeBox(const FVector& Center, const FVector& Extent, const FColor& Color, const FString& Description)
+{
+	FGameplayDebuggerShape NewElement;
+	NewElement.ShapeData.Add(Center);
+	NewElement.ShapeData.Add(Extent);
+	NewElement.Color = Color;
+	NewElement.Description = Description;
+	NewElement.Type = EGameplayDebuggerShape::Box;
+
+	return NewElement;
+}
+
+FGameplayDebuggerShape FGameplayDebuggerShape::MakeCone(const FVector& Location, const FVector& Direction, const float Length, const FColor& Color, const FString& Description)
+{
+	FGameplayDebuggerShape NewElement;
+	NewElement.ShapeData.Add(Location);
+	NewElement.ShapeData.Add(Direction);
+	NewElement.ShapeData.Add(FVector(Length, 0, 0));
+	NewElement.Color = Color;
+	NewElement.Description = Description;
+	NewElement.Type = EGameplayDebuggerShape::Cone;
+
+	return NewElement;
+}
+
+FGameplayDebuggerShape FGameplayDebuggerShape::MakeCylinder(const FVector& Center, const float Radius, const float HalfHeight, const FColor& Color, const FString& Description)
+{
+	FGameplayDebuggerShape NewElement;
+	NewElement.ShapeData.Add(Center);
+	NewElement.ShapeData.Add(FVector(Radius, 0, HalfHeight));
+	NewElement.Color = Color;
+	NewElement.Description = Description;
+	NewElement.Type = EGameplayDebuggerShape::Cylinder;
+
+	return NewElement;
+}
+
+FGameplayDebuggerShape FGameplayDebuggerShape::MakeCircle(const FVector& Center, const FVector& Up, const float Radius, const FColor& Color, const FString& Description)
+{
+	return MakeCircle(Center, Up, Radius, 1.f, Color, Description);
+}
+
+FGameplayDebuggerShape FGameplayDebuggerShape::MakeCircle(const FVector& Center, const FVector& Up, const float Radius, const float Thickness, const FColor& Color, const FString& Description)
+{
+	FGameplayDebuggerShape NewElement;
+	NewElement.ShapeData.Add(Center);
+	NewElement.ShapeData.Add(Up);
+	NewElement.ShapeData.Add(FVector(Radius, Thickness, 0));
+	NewElement.Color = Color;
+	NewElement.Description = Description;
+	NewElement.Type = EGameplayDebuggerShape::Circle;
+
+	return NewElement;
+}
+
+FGameplayDebuggerShape FGameplayDebuggerShape::MakeCapsule(const FVector& Center, const float Radius, const float HalfHeight, const FColor& Color, const FString& Description)
+{
+	FGameplayDebuggerShape NewElement;
+	NewElement.ShapeData.Add(Center);
+	NewElement.ShapeData.Add(FVector(Radius, 0, HalfHeight));
+	NewElement.Color = Color;
+	NewElement.Description = Description;
+	NewElement.Type = EGameplayDebuggerShape::Capsule;
+
+	return NewElement;
+}
+
+FGameplayDebuggerShape FGameplayDebuggerShape::MakePolygon(const TArray<FVector>& Verts, const FColor& Color, const FString& Description)
+{
+	FGameplayDebuggerShape NewElement;
+	NewElement.ShapeData = Verts;
+	NewElement.Color = Color;
+	NewElement.Description = Description;
+	NewElement.Type = EGameplayDebuggerShape::Polygon;
+	return NewElement;
+}
+
+void FGameplayDebuggerShape::Draw(UWorld* World, FGameplayDebuggerCanvasContext& Context)
+{
+	constexpr bool bPersistent = false;
+	constexpr float LifeTime = -1.0f;
+	constexpr uint8 DepthPriority = SDPG_World;
+	
+	FVector DescLocation;
+	switch (Type)
+	{
+	case EGameplayDebuggerShape::Point:
+		if (ShapeData.Num() == 2 && ShapeData[1].X > 0)
+		{
+			DrawDebugSphere(World, ShapeData[0], ShapeData[1].X, 16, Color, bPersistent, LifeTime, DepthPriority);
+			DescLocation = ShapeData[0];
+		}
+		break;
+
+	case EGameplayDebuggerShape::Segment:
+		if (ShapeData.Num() == 3 && ShapeData[2].X > 0)
+		{
+			DrawDebugLine(World, ShapeData[0], ShapeData[1], Color, bPersistent, LifeTime, DepthPriority, ShapeData[2].X);
+			DescLocation = (ShapeData[0] + ShapeData[1]) * 0.5f;
+		}
+		break;
+
+	case EGameplayDebuggerShape::Arrow:
+		if (ShapeData.Num() == 3 && ShapeData[2].X > 0)
+		{
+			DrawDebugDirectionalArrow(World, ShapeData[0], ShapeData[1], ShapeData[2].Y, Color, bPersistent, LifeTime, DepthPriority, ShapeData[2].X);
+			DescLocation = (ShapeData[0] + ShapeData[1]) * 0.5f;
+		}
+		break;
+
+	case EGameplayDebuggerShape::Box:
+		if (ShapeData.Num() == 2)
+		{
+			DrawDebugBox(World, ShapeData[0], ShapeData[1], Color);
+			DescLocation = ShapeData[0];
+		}
+		break;
+
+	case EGameplayDebuggerShape::Cone:
+		if (ShapeData.Num() == 3 && ShapeData[2].X > 0)
+		{
+			const float DefaultConeAngle = 0.25f; // ~ 15 degrees
+			DrawDebugCone(World, ShapeData[0], ShapeData[1], ShapeData[2].X, DefaultConeAngle, DefaultConeAngle, 16, Color, bPersistent, LifeTime, DepthPriority);
+			DescLocation = ShapeData[0];
+		}
+		break;
+
+	case EGameplayDebuggerShape::Cylinder:
+		if (ShapeData.Num() == 2)
+		{
+			DrawDebugCylinder(World, ShapeData[0] - FVector(0, 0, ShapeData[1].Z), ShapeData[0] + FVector(0, 0, ShapeData[1].Z), ShapeData[1].X, 16, Color, bPersistent, LifeTime, DepthPriority);
+			DescLocation = ShapeData[0];
+		}
+		break;
+
+	case EGameplayDebuggerShape::Circle:
+		if (ShapeData.Num() == 3)
+		{
+			const FMatrix Axes = FRotationMatrix::MakeFromX(ShapeData[1]);
+			DrawDebugCircle(World, ShapeData[0], ShapeData[2].X, 32, Color, bPersistent, LifeTime, DepthPriority, ShapeData[2].Y, Axes.GetUnitAxis(EAxis::Y), Axes.GetUnitAxis(EAxis::Z), false);
+			DescLocation = ShapeData[0];
+		}
+		break;
+
+	case EGameplayDebuggerShape::Capsule:
+		if (ShapeData.Num() == 2)
+		{
+			DrawDebugCapsule(World, ShapeData[0], ShapeData[1].Z, ShapeData[1].X, FQuat::Identity, Color, bPersistent, LifeTime, DepthPriority);
+			DescLocation = ShapeData[0];
+		}
+		break;
+
+	case EGameplayDebuggerShape::Polygon:
+		if (ShapeData.Num() > 0)
+		{
+			FVector MidPoint = FVector::ZeroVector;
+			TArray<int32> Indices;
+			for (int32 Idx = 0; Idx < ShapeData.Num(); Idx++)
+			{
+				Indices.Add(Idx);
+				MidPoint += ShapeData[Idx];
+			}
+
+			DrawDebugMesh(World, ShapeData, Indices, Color, bPersistent, LifeTime, DepthPriority);
+			DescLocation = MidPoint / ShapeData.Num();
+		}
+		break;
+
+	default:
+		break;
+	}
+
+	if (Description.Len() && Context.IsLocationVisible(DescLocation))
+	{
+		const FVector2D ScreenLoc = Context.ProjectLocation(DescLocation);
+		Context.PrintAt(ScreenLoc.X, ScreenLoc.Y, Color, Description);
+	}
+}
+
+FArchive& operator<<(FArchive& Ar, FGameplayDebuggerShape& Shape)
+{
+	Ar << Shape.ShapeData;
+	Ar << Shape.Description;
+	Ar << Shape.Color;
+
+	uint8 TypeNum = static_cast<uint8>(Shape.Type);
+	Ar << TypeNum;
+	Shape.Type = static_cast<EGameplayDebuggerShape>(TypeNum);
+
+	return Ar;
+}
+
+//////////////////////////////////////////////////////////////////////////
+// FGameplayDebuggerCanvasContext
+
+enum class EStringParserToken : uint8
+{
+	OpenTag,
+	CloseTag,
+	NewLine,
+	EndOfString,
+	RegularChar,
+	Tab,
+};
+
+class FTaggedStringParser
+{
+public:
+	struct FNode
+	{
+		FString String;
+		FColor Color;
+		bool bNewLine;
+
+		FNode() : Color(FColor::White), bNewLine(false) {}
+		FNode(const FColor& InColor) : Color(InColor), bNewLine(false) {}
+	};
+
+	TArray<FNode> NodeList;
+
+	FTaggedStringParser(const FColor& InDefaultColor) : DefaultColor(InDefaultColor) {}
+
+	void ParseString(const FString& StringToParse)
+	{
+		DataIndex = 0;
+		DataString = StringToParse;
+		if (DataIndex >= DataString.Len())
+		{
+			return;
+		}
+
+		const FString TabString(TEXT("     "));
+		FColor TagColor;
+		FNode CurrentNode(DefaultColor);
+
+		for (EStringParserToken Token = ReadToken(); Token != EStringParserToken::EndOfString; Token = ReadToken())
+		{
+			switch (Token)
+			{
+			case EStringParserToken::RegularChar:
+				CurrentNode.String.AppendChar(DataString[DataIndex]);
+				break;
+
+			case EStringParserToken::NewLine:
+				NodeList.Add(CurrentNode);
+				CurrentNode = FNode(NodeList.Last().Color);
+				CurrentNode.bNewLine = true;
+				break;
+
+			case EStringParserToken::Tab:
+				CurrentNode.String.Append(TabString);
+				break;
+
+			case EStringParserToken::OpenTag:
+				if (ParseTag(TagColor))
+				{
+					NodeList.Add(CurrentNode);
+					CurrentNode = FNode(TagColor);
+				}
+				break;
+			}
+
+			DataIndex++;
+		}
+
+		NodeList.Add(CurrentNode);
+	}
+
+private:
+
+	int32 DataIndex;
+	FString DataString;
+	FColor DefaultColor;
+
+	EStringParserToken ReadToken() const
+	{
+		EStringParserToken OutToken = EStringParserToken::RegularChar;
+
+		const TCHAR Char = DataIndex < DataString.Len() ? DataString[DataIndex] : TEXT('\0');
+		switch (Char)
+		{
+		case TEXT('\0'):
+			OutToken = EStringParserToken::EndOfString;
+			break;
+		case TEXT('{'):
+			OutToken = EStringParserToken::OpenTag;
+			break;
+		case TEXT('}'):
+			OutToken = EStringParserToken::CloseTag;
+			break;
+		case TEXT('\n'):
+			OutToken = EStringParserToken::NewLine;
+			break;
+		case TEXT('\t'):
+			OutToken = EStringParserToken::Tab;
+			break;
+		default:
+			break;
+		}
+
+		return OutToken;
+	}
+
+	bool ParseTag(FColor& OutColor)
+	{
+		FString TagString;
+
+		EStringParserToken Token = ReadToken();
+		for (; Token != EStringParserToken::EndOfString && Token != EStringParserToken::CloseTag; Token = ReadToken())
+		{
+			if (Token == EStringParserToken::RegularChar)
+			{
+				TagString.AppendChar(DataString[DataIndex]);
+			}
+			
+			DataIndex++;
+		}
+
+		bool bResult = false;
+		if (Token == EStringParserToken::CloseTag)
+		{
+			const FString TagColorLower = TagString.ToLower();
+			const bool bIsColorName = GColorList.IsValidColorName(*TagColorLower);
+			
+			if (bIsColorName)
+			{
+				OutColor = GColorList.GetFColorByName(*TagColorLower);
+				bResult = true;
+			}
+			else
+			{
+				bResult = OutColor.InitFromString(TagString);
+			}
+		}
+
+		return bResult;
+	}
+};
+
+FGameplayDebuggerCanvasContext::FGameplayDebuggerCanvasContext(UCanvas* InCanvas, UFont* InFont)
+{
+	if (InCanvas)
+	{
+		Canvas = InCanvas;
+		Font = InFont;
+		CursorX = DefaultX = InCanvas->SafeZonePadX;
+		CursorY = DefaultY = InCanvas->SafeZonePadY;
+	}
+	else
+	{
+		CursorX = DefaultX = 0.0f;
+		CursorY = DefaultY = 0.0f;
+	}
+}
+
+void FGameplayDebuggerCanvasContext::Print(const FString& String)
+{
+	Print(FColor::White, 1.0f, String);
+}
+
+void FGameplayDebuggerCanvasContext::Print(const FColor& Color, const FString& String)
+{
+	Print(Color,  1.0f, String);
+}
+
+void FGameplayDebuggerCanvasContext::Print(const FColor& Color, const float Alpha, const FString& String)
+{
+	FTaggedStringParser Parser(Color);
+	Parser.ParseString(String);
+
+	const float LineHeight = GetLineHeight();
+	for (int32 NodeIdx = 0; NodeIdx < Parser.NodeList.Num(); NodeIdx++)
+	{
+		const FTaggedStringParser::FNode& NodeData = Parser.NodeList[NodeIdx];
+		if (NodeData.bNewLine)
+		{
+			if (Canvas.IsValid() && (CursorY + LineHeight) > Canvas->ClipY)
+			{
+				DefaultX += Canvas->ClipX / 2;
+				CursorY = 0.0f;
+			}
+
+			CursorX = DefaultX;
+			CursorY += LineHeight;
+		}
+
+		if (NodeData.String.Len() > 0)
+		{
+			float SizeX = 0.0f, SizeY = 0.0f;
+			MeasureString(NodeData.String, SizeX, SizeY);
+
+			FLinearColor TextColor(NodeData.Color);
+			TextColor.A = Alpha;
+			
+			FCanvasTextItem TextItem(FVector2D::ZeroVector, FText::FromString(NodeData.String), Font.Get(), TextColor);
+			if (FontRenderInfo.bEnableShadow)
+			{
+				TextItem.EnableShadow(FColor::Black, FVector2D(1, 1));
+			}
+
+			DrawItem(TextItem, CursorX, CursorY);
+			CursorX += SizeX;
+		}
+	}
+
+	MoveToNewLine();
+}
+
+void FGameplayDebuggerCanvasContext::PrintAt(float PosX, float PosY, const FString& String)
+{
+	PrintAt(PosX, PosY, FColor::White, 1.0f, String);
+}
+
+void FGameplayDebuggerCanvasContext::PrintAt(float PosX, float PosY, const FColor& Color, const FString& String)
+{
+	PrintAt(PosX, PosY, Color, 1.0f, String);
+}
+
+void FGameplayDebuggerCanvasContext::PrintAt(float PosX, float PosY, const FColor& Color, const float Alpha, const FString& String)
+{
+	TGuardValue<float> ScopedCursorX(CursorX, PosX);
+	TGuardValue<float> ScopedCursorY(CursorY, PosY);
+	TGuardValue<float> ScopedDefaultX(DefaultX, PosX);
+	TGuardValue<float> ScopedDefaultY(DefaultY, PosY);
+
+	Print(Color, Alpha, String);
+}
+
+// copied from Core/Private/Misc/VarargsHeler.h 
+#define GROWABLE_PRINTF(PrintFunc) \
+	int32	BufferSize	= 1024; \
+	TCHAR*	Buffer		= NULL; \
+	int32	Result		= -1; \
+	/* allocate some stack space to use on the first pass, which matches most strings */ \
+	TCHAR	StackBuffer[512]; \
+	TCHAR*	AllocatedBuffer = NULL; \
+\
+	/* first, try using the stack buffer */ \
+	Buffer = StackBuffer; \
+	GET_VARARGS_RESULT( Buffer, UE_ARRAY_COUNT(StackBuffer), UE_ARRAY_COUNT(StackBuffer) - 1, Fmt, Fmt, Result ); \
+\
+	/* if that fails, then use heap allocation to make enough space */ \
+	while(Result == -1) \
+	{ \
+		FMemory::SystemFree(AllocatedBuffer); \
+		/* We need to use malloc here directly as GMalloc might not be safe. */ \
+		Buffer = AllocatedBuffer = (TCHAR*) FMemory::SystemMalloc( BufferSize * sizeof(TCHAR) ); \
+		GET_VARARGS_RESULT( Buffer, BufferSize, BufferSize-1, Fmt, Fmt, Result ); \
+		BufferSize *= 2; \
+	}; \
+	Buffer[Result] = 0; \
+	; \
+\
+	PrintFunc; \
+	FMemory::SystemFree(AllocatedBuffer);
+
+void FGameplayDebuggerCanvasContext::PrintfImpl(const FColor& Color, const float Alpha, const TCHAR* Fmt, ...)
+{
+	GROWABLE_PRINTF(Print(Color, Alpha, Buffer));
+}
+
+void FGameplayDebuggerCanvasContext::PrintfAtImpl(float PosX, float PosY, const FColor& Color, const float Alpha, const TCHAR* Fmt, ...)
+{
+	GROWABLE_PRINTF(PrintAt(PosX, PosY, Color, Alpha, Buffer));
+}
+
+void FGameplayDebuggerCanvasContext::MoveToNewLine()
+{
+	const float LineHeight = GetLineHeight();
+	CursorY += LineHeight;
+	CursorX = DefaultX;
+}
+
+void FGameplayDebuggerCanvasContext::MeasureString(const FString& String, float& OutSizeX, float& OutSizeY) const
+{
+	OutSizeX = OutSizeY = 0.0f;
+
+	UCanvas* CanvasOb = Canvas.Get();
+	if (CanvasOb)
+	{
+		FString StringWithoutFormatting = String;
+		int32 BracketStart = INDEX_NONE;
+		while (StringWithoutFormatting.FindChar(TEXT('{'), BracketStart))
+		{
+			int32 BracketEnd = INDEX_NONE;
+			if (StringWithoutFormatting.FindChar(TEXT('}'), BracketEnd))
+			{
+				if (BracketEnd > BracketStart)
+				{
+					StringWithoutFormatting.RemoveAt(BracketStart, BracketEnd - BracketStart + 1, false);
+				}
+			}
+		}
+
+		TArray<FString> Lines;
+		StringWithoutFormatting.ParseIntoArrayLines(Lines);
+		
+		UFont* FontOb = Font.Get();
+		for (int32 Idx = 0; Idx < Lines.Num(); Idx++)
+		{
+			float LineSizeX = 0.0f, LineSizeY = 0.0f;
+			CanvasOb->StrLen(FontOb, Lines[Idx], LineSizeX, LineSizeY);
+
+			OutSizeX = FMath::Max(OutSizeX, LineSizeX);
+			OutSizeY += LineSizeY;
+		}
+	}
+}
+
+float FGameplayDebuggerCanvasContext::GetLineHeight() const
+{
+	UFont* FontOb = Font.Get();
+	return FontOb ? FontOb->GetMaxCharHeight() : 0.0f;
+}
+
+FVector2D FGameplayDebuggerCanvasContext::ProjectLocation(const FVector& Location) const
+{
+	UCanvas* CanvasOb = Canvas.Get();
+	return CanvasOb ? FVector2D(CanvasOb->Project(Location)) : FVector2D::ZeroVector;
+}
+
+bool FGameplayDebuggerCanvasContext::IsLocationVisible(const FVector& Location) const
+{
+	return Canvas.IsValid() && Canvas->SceneView && Canvas->SceneView->ViewFrustum.IntersectSphere(Location, 1.0f);
+}
+
+void FGameplayDebuggerCanvasContext::DrawItem(FCanvasItem& Item, float PosX, float PosY)
+{
+	UCanvas* CanvasOb = Canvas.Get();
+	if (CanvasOb)
+	{
+		CanvasOb->DrawItem(Item, PosX, PosY);
+	}
+}
+
+void FGameplayDebuggerCanvasContext::DrawIcon(const FColor& Color, const FCanvasIcon& Icon, float PosX, float PosY, float Scale)
+{
+	UCanvas* CanvasOb = Canvas.Get();
+	if (CanvasOb)
+	{
+		CanvasOb->SetDrawColor(Color);
+		CanvasOb->DrawIcon(Icon, PosX, PosY, Scale);
+	}
+}
+
+UWorld* FGameplayDebuggerCanvasContext::GetWorld() const 
+{
+	return World.IsValid() 
+		? World.Get()
+		: (PlayerController.IsValid() ? PlayerController->GetWorld() : nullptr);
+}
+
+//////////////////////////////////////////////////////////////////////////
+// FGameplayDebuggerDataPack
+
+int32 FGameplayDebuggerDataPack::PacketSize = 512;
+
+bool FGameplayDebuggerDataPack::CheckDirtyAndUpdate()
+{
+	TArray<uint8> UncompressedBuffer;
+	FMemoryWriter ArWriter(UncompressedBuffer);
+	SerializeDelegate.Execute(ArWriter);
+
+	const uint32 NewDataCRC = FCrc::MemCrc32(UncompressedBuffer.GetData(), UncompressedBuffer.Num());
+	if ((NewDataCRC == DataCRC) && !bIsDirty)
+	{
+		return false;
+	}
+
+	DataCRC = NewDataCRC;
+	return true;
+}
+
+bool FGameplayDebuggerDataPack::RequestReplication(int16 SyncCounter)
+{
+	if (bNeedsConfirmation && !bReceived)
+	{
+		return false;
+	}
+
+	TArray<uint8> UncompressedBuffer;
+	FMemoryWriter ArWriter(UncompressedBuffer);
+	SerializeDelegate.Execute(ArWriter);
+
+	const uint32 NewDataCRC = FCrc::MemCrc32(UncompressedBuffer.GetData(), UncompressedBuffer.Num());
+	if ((NewDataCRC == DataCRC) && !bIsDirty)
+	{
+		return false;
+	}
+
+	const int32 MaxUncompressedDataSize = PacketSize;
+
+	Header.bIsCompressed = (UncompressedBuffer.Num() > MaxUncompressedDataSize);
+	if (Header.bIsCompressed)
+	{
+		const int32 UncompressedSize = UncompressedBuffer.Num();
+
+		int32 CompressionHeader = UncompressedSize;
+		const int32 CompressionHeaderSize = sizeof(CompressionHeader);
+
+		int32 CompressedSize = FMath::TruncToInt(1.1f * UncompressedSize);
+		Data.SetNum(CompressionHeaderSize + CompressedSize);
+
+		uint8* CompressedBuffer = Data.GetData();
+		FMemory::Memcpy(CompressedBuffer, &CompressionHeader, CompressionHeaderSize);
+		CompressedBuffer += CompressionHeaderSize;
+
+		FCompression::CompressMemory(NAME_Zlib,	CompressedBuffer, CompressedSize, UncompressedBuffer.GetData(), UncompressedSize, COMPRESS_BiasMemory);
+
+		Data.SetNum(CompressionHeaderSize + CompressedSize);
+	}
+	else
+	{
+		Data = UncompressedBuffer;
+	}
+
+	bNeedsConfirmation = IsMultiPacket(Data.Num());
+	bReceived = false;
+	bIsDirty = false;
+
+	DataCRC = NewDataCRC;
+	Header.DataOffset = 0;
+	Header.DataSize = Data.Num();
+	Header.SyncCounter = SyncCounter;
+	Header.DataVersion++;
+	return true;
+}
+
+void FGameplayDebuggerDataPack::OnReplicated()
+{
+	if (Header.DataSize == 0)
+	{
+		ResetDelegate.Execute();
+		return;
+	}
+
+	if (Header.bIsCompressed)
+	{
+		uint8* CompressedBuffer = Data.GetData();
+		int32 CompressionHeader = 0;
+		const int32 CompressionHeaderSize = sizeof(CompressionHeader);
+
+		FMemory::Memcpy(&CompressionHeader, CompressedBuffer, CompressionHeaderSize);
+		CompressedBuffer += CompressionHeaderSize;
+
+		const int32 CompressedSize = Data.Num() - CompressionHeaderSize;
+		const int32 UncompressedSize = CompressionHeader;
+		TArray<uint8> UncompressedBuffer;
+		UncompressedBuffer.AddUninitialized(UncompressedSize);
+
+		FCompression::UncompressMemory(NAME_Zlib, UncompressedBuffer.GetData(), UncompressedSize, CompressedBuffer, CompressedSize);
+
+		FMemoryReader ArReader(UncompressedBuffer);
+		SerializeDelegate.Execute(ArReader);
+	}
+	else
+	{
+		FMemoryReader ArReader(Data);
+		SerializeDelegate.Execute(ArReader);
+	}
+
+	Header.DataOffset = Header.DataSize;
+}
+
+void FGameplayDebuggerDataPack::OnPacketRequest(int16 DataVersion, int32 DataOffset)
+{
+	// client should confirm with the same version and offset that server currently replicates
+	if (DataVersion == Header.DataVersion && DataOffset == Header.DataOffset)
+	{
+		Header.DataOffset = FMath::Min(DataOffset + FGameplayDebuggerDataPack::PacketSize, Header.DataSize);
+		bReceived = (Header.DataOffset == Header.DataSize);
+	}
+	// if for some reason it requests previous data version, rollback to first packet
+	else if (DataVersion < Header.DataVersion)
+	{
+		Header.DataOffset = 0;
+	}
+	// it may also request a previous packet from the same version, rollback and send next one
+	else if (DataVersion == Header.DataVersion && DataOffset < Header.DataOffset)
+	{
+		Header.DataOffset = FMath::Max(0, DataOffset + FGameplayDebuggerDataPack::PacketSize);
+	}
+}
+
+//////////////////////////////////////////////////////////////////////////
+// FGameplayDebuggerInputModifier
+
+FGameplayDebuggerInputModifier FGameplayDebuggerInputModifier::Shift(true, false, false, false);
+FGameplayDebuggerInputModifier FGameplayDebuggerInputModifier::Ctrl(false, true, false, false);
+FGameplayDebuggerInputModifier FGameplayDebuggerInputModifier::Alt(false, false, true, false);
+FGameplayDebuggerInputModifier FGameplayDebuggerInputModifier::Cmd(false, false, false, true);
+FGameplayDebuggerInputModifier FGameplayDebuggerInputModifier::None;
+
+//////////////////////////////////////////////////////////////////////////
+// FGameplayDebuggerInputHandler
+
+bool FGameplayDebuggerInputHandler::IsValid() const
+{
+	return FKey(KeyName).IsValid();
+}
+
+FString FGameplayDebuggerInputHandler::ToString() const
+{
+	FString KeyDesc = KeyName.ToString();
+	
+	if (Modifier.bShift)
+	{
+		KeyDesc = FString(TEXT("Shift+")) + KeyDesc;
+	}
+
+	if (Modifier.bAlt)
+	{
+		KeyDesc = FString(TEXT("Alt+")) + KeyDesc;
+	}
+	
+	if (Modifier.bCtrl)
+	{
+		KeyDesc = FString(TEXT("Ctrl+")) + KeyDesc;
+	}
+	
+	if (Modifier.bCmd)
+	{
+		KeyDesc = FString(TEXT("Cmd+")) + KeyDesc;
+	}
+
+	return KeyDesc;
+}
+
+//////////////////////////////////////////////////////////////////////////
+// FGameplayDebuggerInputHandlerConfig
+
+FName FGameplayDebuggerInputHandlerConfig::CurrentCategoryName;
+FName FGameplayDebuggerInputHandlerConfig::CurrentExtensionName;
+
+FGameplayDebuggerInputHandlerConfig::FGameplayDebuggerInputHandlerConfig(const FName ConfigName, const FName DefaultKeyName)
+{
+	KeyName = DefaultKeyName;
+	UpdateConfig(ConfigName);
+}
+
+FGameplayDebuggerInputHandlerConfig::FGameplayDebuggerInputHandlerConfig(const FName ConfigName, const FName DefaultKeyName, const FGameplayDebuggerInputModifier& DefaultModifier)
+{
+	KeyName = DefaultKeyName;
+	Modifier = DefaultModifier;
+	UpdateConfig(ConfigName);
+}
+
+void FGameplayDebuggerInputHandlerConfig::UpdateConfig(const FName ConfigName)
+{
+	if (FGameplayDebuggerInputHandlerConfig::CurrentCategoryName != NAME_None)
+	{
+		UGameplayDebuggerConfig* MutableToolConfig = UGameplayDebuggerConfig::StaticClass()->GetDefaultObject<UGameplayDebuggerConfig>();
+		MutableToolConfig->UpdateCategoryInputConfig(FGameplayDebuggerInputHandlerConfig::CurrentCategoryName, ConfigName, KeyName, Modifier);
+	}
+	else if (FGameplayDebuggerInputHandlerConfig::CurrentExtensionName != NAME_None)
+	{
+		UGameplayDebuggerConfig* MutableToolConfig = UGameplayDebuggerConfig::StaticClass()->GetDefaultObject<UGameplayDebuggerConfig>();
+		MutableToolConfig->UpdateExtensionInputConfig(FGameplayDebuggerInputHandlerConfig::CurrentExtensionName, ConfigName, KeyName, Modifier);
+	}
+}
diff --git a/Engine/Source/Runtime/GameplayDebugger/Public/GameplayDebugger.h b/Engine/Source/Runtime/GameplayDebugger/Public/GameplayDebugger.h
new file mode 100644
--- /dev/null
+++ b/Engine/Source/Runtime/GameplayDebugger/Public/GameplayDebugger.h
@@ -0,0 +1,82 @@
+// Copyright Epic Games, Inc. All Rights Reserved.
+
+// GAMEPLAY DEBUGGER
+// 
+// This tool allows easy on screen debugging of gameplay data, supporting client-server replication.
+// Data is organized into named categories, which can be toggled during debugging.
+// 
+// To enable it, press Apostrophe key (UGameplayDebuggerConfig::ActivationKey)
+//
+// Category class:
+// - derives from FGameplayDebuggerCategory
+// - implements at least CollectData() and DrawData() functions
+// - requires WITH_GAMEPLAY_DEBUGGER define to compile (doesn't exist in shipping builds by default; see note below to override)
+// - needs to be registered and unregistered manually by owning module
+// - automatically replicate data added with FGameplayDebuggerCategory::AddTextLine, FGameplayDebuggerCategory::AddShape
+// - automatically replicate data structs initialized with FGameplayDebuggerCategory::SetDataPackReplication
+// - can define own input bindings (e.g. subcategories, etc)
+//
+// Extension class:
+// - derives from FGameplayDebuggerExtension
+// - needs to be registered and unregistered manually by owning module
+// - can define own input bindings
+// - basically it's a stateless, not replicated, not drawn category, ideal for making e.g. different actor selection mechanic
+//
+// 
+// Check FGameplayDebuggerCategory_BehaviorTree for implementation example.
+// Check AIModule/Private/AIModule.cpp for registration example.
+//
+// Note. Use 'SetupGameplayDebuggerSupport(Target)' when adding module to your project's Build.cs (see AIModule/AIModule.Build.cs)
+// Note. Use 'bUseGameplayDebugger={0|1}' in your <ProjectTargetType>.Target.cs to force GameplayDebugger disabled/enabled
+// 
+
+#pragma once
+
+#include "CoreMinimal.h"
+#include "Modules/ModuleInterface.h"
+#include "Modules/ModuleManager.h"
+
+enum class EGameplayDebuggerCategoryState : uint8
+{
+	EnabledInGameAndSimulate,
+	EnabledInGame,
+	EnabledInSimulate,
+	Disabled,
+	Hidden,
+};
+
+class IGameplayDebugger : public IModuleInterface
+{
+
+public:
+	DECLARE_DELEGATE_RetVal(TSharedRef<class FGameplayDebuggerCategory>, FOnGetCategory);
+	DECLARE_DELEGATE_RetVal(TSharedRef<class FGameplayDebuggerExtension>, FOnGetExtension);
+
+	/**
+	 * Singleton-like access to this module's interface.  This is just for convenience!
+	 * Beware of calling this during the shutdown phase, though.  Your module might have been unloaded already.
+	 *
+	 * @return Returns singleton instance, loading the module on demand if needed
+	 */
+	static inline IGameplayDebugger& Get()
+	{
+		return FModuleManager::LoadModuleChecked< IGameplayDebugger >("GameplayDebugger");
+	}
+
+	/**
+	 * Checks to see if this module is loaded and ready.  It is only valid to call Get() if IsAvailable() returns true.
+	 *
+	 * @return True if the module is loaded and ready to use
+	 */
+	static inline bool IsAvailable()
+	{
+		return FModuleManager::Get().IsModuleLoaded("GameplayDebugger");
+	}
+
+	virtual void RegisterCategory(FName CategoryName, FOnGetCategory MakeInstanceDelegate, EGameplayDebuggerCategoryState CategoryState = EGameplayDebuggerCategoryState::Disabled, int32 SlotIdx = INDEX_NONE) = 0;
+	virtual void UnregisterCategory(FName CategoryName) = 0;
+	virtual void NotifyCategoriesChanged() = 0;
+	virtual void RegisterExtension(FName ExtensionName, FOnGetExtension MakeInstanceDelegate) = 0;
+	virtual void UnregisterExtension(FName ExtensionName) = 0;
+	virtual void NotifyExtensionsChanged() = 0;
+};
diff --git a/Engine/Source/Runtime/GameplayDebugger/Public/GameplayDebuggerAddonBase.h b/Engine/Source/Runtime/GameplayDebugger/Public/GameplayDebuggerAddonBase.h
new file mode 100644
--- /dev/null
+++ b/Engine/Source/Runtime/GameplayDebugger/Public/GameplayDebuggerAddonBase.h
@@ -0,0 +1,108 @@
+// Copyright Epic Games, Inc. All Rights Reserved.
+
+#pragma once
+
+#include "CoreMinimal.h"
+#include "GameplayDebuggerTypes.h"
+
+class UWorld;
+class AActor;
+class APlayerController;
+class AGameplayDebuggerCategoryReplicator;
+
+class GAMEPLAYDEBUGGER_API FGameplayDebuggerAddonBase : public TSharedFromThis<FGameplayDebuggerAddonBase>
+{
+public:
+
+	virtual ~FGameplayDebuggerAddonBase() {}
+
+	int32 GetNumInputHandlers() const { return InputHandlers.Num(); }
+	FGameplayDebuggerInputHandler& GetInputHandler(int32 HandlerId) { return InputHandlers[HandlerId]; }
+	FString GetInputHandlerDescription(int32 HandlerId) const;
+	UWorld* GetWorldFromReplicator() const;
+	/** Returns the first non-null world associated with, in order: OwnerPC, DebugActor, AGameplayDebuggerCategoryReplicator owner */
+	UWorld* GetDataWorld(const APlayerController* OwnerPC, const AActor* DebugActor) const;
+
+	/** [ALL] called when gameplay debugger is activated */
+	virtual void OnGameplayDebuggerActivated();
+
+	/** [ALL] called when gameplay debugger is deactivated */
+	virtual void OnGameplayDebuggerDeactivated();
+
+	/** check if simulate in editor mode is active */
+	static bool IsSimulateInEditor();
+
+protected:
+
+	/** tries to find selected actor in local world */
+	AActor* FindLocalDebugActor() const;
+
+	/** returns replicator actor */
+	AGameplayDebuggerCategoryReplicator* GetReplicator() const;
+
+	/** creates new key binding handler: single key press */
+	template<class UserClass>
+	bool BindKeyPress(FName KeyName, UserClass* KeyHandlerObject, typename FGameplayDebuggerInputHandler::FHandler::TMethodPtr< UserClass > KeyHandlerFunc, EGameplayDebuggerInputMode InputMode = EGameplayDebuggerInputMode::Local)
+	{
+		FGameplayDebuggerInputHandler NewHandler;
+		NewHandler.KeyName = KeyName;
+		NewHandler.Delegate.BindRaw(KeyHandlerObject, KeyHandlerFunc);
+		NewHandler.Mode = InputMode;
+
+		if (NewHandler.IsValid())
+		{
+			InputHandlers.Add(NewHandler);
+			return true;
+		}
+
+		return false;
+	}
+
+	/** creates new key binding handler: key press with modifiers */
+	template<class UserClass>
+	bool BindKeyPress(FName KeyName, FGameplayDebuggerInputModifier KeyModifer, UserClass* KeyHandlerObject, typename FGameplayDebuggerInputHandler::FHandler::TMethodPtr< UserClass > KeyHandlerFunc, EGameplayDebuggerInputMode InputMode = EGameplayDebuggerInputMode::Local)
+	{
+		FGameplayDebuggerInputHandler NewHandler;
+		NewHandler.KeyName = KeyName;
+		NewHandler.Modifier = KeyModifer;
+		NewHandler.Delegate.BindRaw(KeyHandlerObject, KeyHandlerFunc);
+		NewHandler.Mode = InputMode;
+
+		if (NewHandler.IsValid())
+		{
+			InputHandlers.Add(NewHandler);
+			return true;
+		}
+
+		return false;
+	}
+
+	/** creates new key binding handler: customizable key press, stored in config files */
+	template<class UserClass>
+	bool BindKeyPress(const FGameplayDebuggerInputHandlerConfig& InputConfig, UserClass* KeyHandlerObject, typename FGameplayDebuggerInputHandler::FHandler::TMethodPtr< UserClass > KeyHandlerFunc, EGameplayDebuggerInputMode InputMode = EGameplayDebuggerInputMode::Local)
+	{
+		FGameplayDebuggerInputHandler NewHandler;
+		NewHandler.KeyName = InputConfig.KeyName;
+		NewHandler.Modifier = InputConfig.Modifier;
+		NewHandler.Delegate.BindRaw(KeyHandlerObject, KeyHandlerFunc);
+		NewHandler.Mode = InputMode;
+
+		if (NewHandler.IsValid())
+		{
+			InputHandlers.Add(NewHandler);
+			return true;
+		}
+
+		return false;
+	}
+
+private:
+
+	friend class FGameplayDebuggerAddonManager;
+
+	/** list of input handlers */
+	TArray<FGameplayDebuggerInputHandler> InputHandlers;
+
+	/** replicator actor */
+	TWeakObjectPtr<AGameplayDebuggerCategoryReplicator> RepOwner;
+};
diff --git a/Engine/Source/Runtime/GameplayDebugger/Public/GameplayDebuggerAddonManager.h b/Engine/Source/Runtime/GameplayDebugger/Public/GameplayDebuggerAddonManager.h
new file mode 100644
--- /dev/null
+++ b/Engine/Source/Runtime/GameplayDebugger/Public/GameplayDebuggerAddonManager.h
@@ -0,0 +1,97 @@
+// Copyright Epic Games, Inc. All Rights Reserved.
+
+#pragma once
+
+#include "CoreMinimal.h"
+#include "GameplayDebuggerCategory.h"
+#include "GameplayDebugger.h"
+
+class AGameplayDebuggerCategoryReplicator;
+class FGameplayDebuggerExtension;
+
+DECLARE_MULTICAST_DELEGATE(FOnGameplayDebuggerAddonEvent);
+
+class FGameplayDebuggerCategory;
+class AGameplayDebuggerCategoryReplicator;
+
+struct FGameplayDebuggerCategoryInfo
+{
+	IGameplayDebugger::FOnGetCategory MakeInstanceDelegate;
+	EGameplayDebuggerCategoryState DefaultCategoryState;
+	EGameplayDebuggerCategoryState CategoryState;
+	int32 SlotIdx;
+};
+
+struct FGameplayDebuggerExtensionInfo
+{
+	IGameplayDebugger::FOnGetExtension MakeInstanceDelegate;
+	uint32 bDefaultEnabled : 1;
+	uint32 bEnabled : 1;
+};
+
+class FGameplayDebuggerAddonManager
+{
+public:
+	FGameplayDebuggerAddonManager();
+
+	/** adds new category to managed collection */
+	void RegisterCategory(FName CategoryName, IGameplayDebugger::FOnGetCategory MakeInstanceDelegate, EGameplayDebuggerCategoryState CategoryState, int32 SlotIdx);
+
+	/** removes category from managed collection */
+	void UnregisterCategory(FName CategoryName);
+
+	/** notify about change in known categories */
+	void NotifyCategoriesChanged();
+
+	/** creates new category objects for all known types */
+	void CreateCategories(AGameplayDebuggerCategoryReplicator& Owner, TArray<TSharedRef<FGameplayDebuggerCategory> >& CategoryObjects);
+
+	/** adds new extension to managed collection */
+	void RegisterExtension(FName ExtensionName, IGameplayDebugger::FOnGetExtension MakeInstanceDelegate);
+
+	/** removes extension from managed collection */
+	void UnregisterExtension(FName ExtensionName);
+
+	/** notify about change in known extensions */
+	void NotifyExtensionsChanged();
+
+	/** creates new extension objects for all known types */
+	void CreateExtensions(AGameplayDebuggerCategoryReplicator& Replicator, TArray<TSharedRef<FGameplayDebuggerExtension> >& ExtensionObjects);
+
+	/** refresh category and extension data from config */
+	void UpdateFromConfig();
+
+	/** get slot-Id map */
+	const TArray<TArray<int32> >& GetSlotMap() const { return SlotMap; }
+
+	/** get slot-Name map */
+	const TArray<FString> GetSlotNames() const { return SlotNames; }
+
+	/** get number of visible categories */
+	int32 GetNumVisibleCategories() const { return NumVisibleCategories; }
+
+	/** singleton accessor */
+	static FGameplayDebuggerAddonManager& GetCurrent();
+
+	/** event called when CategoryMap changes */
+	FOnGameplayDebuggerAddonEvent OnCategoriesChanged;
+
+	/** event called when ExtensionMap changes */
+	FOnGameplayDebuggerAddonEvent OnExtensionsChanged;
+
+private:
+	/** map of all known extensions indexed by their names */
+	TMap<FName, FGameplayDebuggerExtensionInfo> ExtensionMap;
+
+	/** map of all known categories indexed by their names */
+	TMap<FName, FGameplayDebuggerCategoryInfo> CategoryMap;
+
+	/** list of all slots and their categories Ids */
+	TArray<TArray<int32> > SlotMap;
+
+	/** list of slot names */
+	TArray<FString> SlotNames;
+
+	/** number of categories, excluding hidden ones */
+	int32 NumVisibleCategories;
+};
diff --git a/Engine/Source/Runtime/GameplayDebugger/Public/GameplayDebuggerCategory.h b/Engine/Source/Runtime/GameplayDebugger/Public/GameplayDebuggerCategory.h
new file mode 100644
--- /dev/null
+++ b/Engine/Source/Runtime/GameplayDebugger/Public/GameplayDebuggerCategory.h
@@ -0,0 +1,198 @@
+// Copyright Epic Games, Inc. All Rights Reserved.
+
+
+// GAMEPLAY DEBUGGER CATEGORY
+// 
+// Single category of gameplay debugger tool, responsible for gathering and presenting data.
+// Category instances are created on both server and local sides, and can use replication to
+// show server's state on client.
+// 
+// It should be compiled and used only when module is included, so every category class
+// needs be placed in #if WITH_GAMEPLAY_DEBUGGER block.
+// 
+// 
+// Server side category :
+// - CollectData() function will be called in category having authority (server / standalone)
+// - set CollectDataInterval for adding delay between data collection, default value is 0, meaning: every tick
+// - AddTextLine() and AddShape() adds new data to replicate, both arrays are cleared before calling CollectData()
+// - SetDataPackReplication() marks struct member variable for replication
+// - MarkDataPackDirty() forces data pack replication, sometimes changes can go unnoticed (CRC based)
+// 
+// Local category :
+// - DrawData() function will be called in every tick to present gathered data
+// - everything added by AddTextLine() and AddShape() will be shown before calling DrawData()
+// - CreateSceneProxy() allows creating custom scene proxies, use with MarkRenderStateDirty()
+// - OnDataPackReplicated() notifies about receiving new data, use with MarkRenderStateDirty() if needed
+// - BindKeyPress() allows creating custom key bindings active only when category is being displayed
+// 
+// 
+// Categories needs to be manually registered and unregistered with GameplayDebugger.
+// It's best to do it in owning module's Startup / Shutdown, similar to detail view customizations.
+// Check AIModule.cpp for examples.
+
+#pragma once
+
+#include "CoreMinimal.h"
+#include "GameplayDebuggerTypes.h"
+#include "GameplayDebuggerAddonBase.h"
+
+class AActor;
+class APlayerController;
+class FDebugRenderSceneProxy;
+class UPrimitiveComponent;
+struct FDebugDrawDelegateHelper;
+
+/**
+ * Single category of visual debugger tool
+ */
+class GAMEPLAYDEBUGGER_API FGameplayDebuggerCategory : public FGameplayDebuggerAddonBase
+{
+public:
+
+	FGameplayDebuggerCategory();
+	virtual ~FGameplayDebuggerCategory();
+
+	/** [AUTH] gather data for replication */
+	virtual void CollectData(APlayerController* OwnerPC, AActor* DebugActor);
+
+	/** [LOCAL] draw collected data */
+	virtual void DrawData(APlayerController* OwnerPC, FGameplayDebuggerCanvasContext& CanvasContext);
+
+	/** [LOCAL] creates a scene proxy for more advanced debug rendering */
+	virtual FDebugRenderSceneProxy* CreateDebugSceneProxy(const UPrimitiveComponent* InComponent, FDebugDrawDelegateHelper*& OutDelegateHelper);
+
+	/** [LOCAL] called after successful replication of entire data pack to client */
+	virtual void OnDataPackReplicated(int32 DataPackId);
+
+	/** [AUTH] adds line of text tagged with {color} to replicated data */
+	void AddTextLine(const FString& TextLine);
+
+	/** [AUTH] adds shape to replicated data */
+	void AddShape(const FGameplayDebuggerShape& Shape);
+
+	/** [LOCAL] draw category */
+	void DrawCategory(APlayerController* OwnerPC, FGameplayDebuggerCanvasContext& CanvasContext);
+
+	/** [LOCAL] check if category should be drawn */
+	bool ShouldDrawCategory(bool bHasDebugActor) const { return IsCategoryEnabled() && (!bShowOnlyWithDebugActor || bHasDebugActor); }
+
+	/** [LOCAL] check if data pack replication status  */
+	bool ShouldDrawReplicationStatus() const { return bShowDataPackReplication; }
+
+	/** [ALL] get name of category */
+	FName GetCategoryName() const { return CategoryName; }
+
+	/** [ALL] check if category header should be drawn */
+	bool IsCategoryHeaderVisible() const { return bShowCategoryName; }
+
+	/** [ALL] check if category is enabled */
+	bool IsCategoryEnabled() const { return bIsEnabled; }
+
+	/** [ALL] check if category is local (present data) */
+	bool IsCategoryLocal() const { return bIsLocal; }
+
+	/** [ALL] check if category has authority (collects data) */
+	bool IsCategoryAuth() const { return bHasAuthority; }
+
+	int32 GetNumDataPacks() const { return ReplicatedDataPacks.Num(); }
+	float GetDataPackProgress(int32 DataPackId) const { return ReplicatedDataPacks.IsValidIndex(DataPackId) ? ReplicatedDataPacks[DataPackId].GetProgress() : 0.0f; }
+	bool IsDataPackReplicating(int32 DataPackId) const { return ReplicatedDataPacks.IsValidIndex(DataPackId) && ReplicatedDataPacks[DataPackId].IsInProgress(); }
+	FGameplayDebuggerDataPack::FHeader GetDataPackHeaderCopy(int32 DataPackId) const { return ReplicatedDataPacks.IsValidIndex(DataPackId) ? ReplicatedDataPacks[DataPackId].Header : FGameplayDebuggerDataPack::FHeader(); }
+
+	// temporary functions for compatibility, will be removed soon
+	TArray<FString> GetReplicatedLinesCopy() const { return ReplicatedLines; }
+	TArray<FGameplayDebuggerShape> GetReplicatedShapesCopy() const { return ReplicatedShapes; }
+
+protected:
+	/**
+	 * Returns view location and direction from replicated data if available or will extract it from
+	 * the provided player controller. The replicated data is used first since it might be provided
+	 * from a detached camera.
+	 * @return whether OutViewLocation and OutViewDirection were assigned to valid values
+	 * @note Method is expected to always return valid outputs if valid controller is provided. Otherwise it
+	 * depends if the view information has been replicated before the methods gets called.
+	 */
+	bool GetViewPoint(const APlayerController* OwnerPC, FVector& OutViewLocation, FVector& OutViewDirection) const;
+
+	/**
+	 * Indicates if a given location is within a vision cone built from provided view location and direction based
+	 * on MaxViewDistance and MaxViewAngle from GameplayDebuggerUserSettings
+	 */
+	static bool IsLocationInViewCone(const FVector& ViewLocation, const FVector& ViewDirection, const FVector& TargetLocation);
+
+	/** [AUTH] marks data pack as needing replication */
+	void MarkDataPackDirty(int32 DataPackId);
+
+	/** [LOCAL] requests new scene proxy */
+	void MarkRenderStateDirty();
+
+	/** [LOCAL] preferred view flag for creating scene proxy */
+	FString GetSceneProxyViewFlag() const;
+
+	/** [ALL] sets up DataPack replication, needs address of property holding data, DataPack's struct must define Serialize(FArchive& Ar) function
+	 *  returns DataPackId
+	 */
+	template<typename T>
+	int32 SetDataPackReplication(T* DataPackAddr, EGameplayDebuggerDataPack Flags = EGameplayDebuggerDataPack::ResetOnTick)
+	{
+		FGameplayDebuggerDataPack NewDataPack;
+		NewDataPack.PackId = ReplicatedDataPacks.Num();
+		NewDataPack.Flags = Flags;
+		NewDataPack.SerializeDelegate.BindRaw(DataPackAddr, &T::Serialize);
+		NewDataPack.ResetDelegate.BindLambda([=] { *DataPackAddr = T(); });
+
+		ReplicatedDataPacks.Add(NewDataPack);
+		return NewDataPack.PackId;
+	}
+
+	/** [AUTH] force data collection on next update */
+	void ForceImmediateCollect() { LastCollectDataTime = -FLT_MAX; }
+
+	/** update interval, 0 = each tick */
+	float CollectDataInterval;
+
+	/** include data pack replication details in drawn messages */
+	uint32 bShowDataPackReplication : 1;
+
+	/** include remaining time to next data collection in drawn messages */
+	uint32 bShowUpdateTimer : 1;
+
+	/** include category name in drawn messages */
+	uint32 bShowCategoryName : 1;
+
+	/** draw category only when DebugActor is present */
+	uint32 bShowOnlyWithDebugActor : 1;
+
+private:
+
+	friend class FGameplayDebuggerAddonManager;
+	friend class AGameplayDebuggerCategoryReplicator;
+	friend struct FGameplayDebuggerNetPack;
+
+	/** if set, this category object can display data */
+	uint32 bIsLocal : 1;
+
+	/** if set, this category object can collect data */
+	uint32 bHasAuthority : 1;
+
+	/** if set, this category object is enabled in debugger */
+	uint32 bIsEnabled : 1;
+
+	/** Id number assigned to this category object */
+	int32 CategoryId;
+
+	/** timestamp of last update */
+	float LastCollectDataTime;
+
+	/** name of debugger category (auto assigned during category registration) */
+	FName CategoryName;
+
+	/** list of replicated text lines, will be reset before CollectData call on AUTH */
+	TArray<FString> ReplicatedLines;
+
+	/** list of replicated shapes, will be reset before CollectData call on AUTH */
+	TArray<FGameplayDebuggerShape> ReplicatedShapes;
+
+	/** list of replicated data packs */
+	TArray<FGameplayDebuggerDataPack> ReplicatedDataPacks;
+};
diff --git a/Engine/Source/Runtime/GameplayDebugger/Public/GameplayDebuggerCategoryReplicator.h b/Engine/Source/Runtime/GameplayDebugger/Public/GameplayDebuggerCategoryReplicator.h
new file mode 100644
--- /dev/null
+++ b/Engine/Source/Runtime/GameplayDebugger/Public/GameplayDebuggerCategoryReplicator.h
@@ -0,0 +1,310 @@
+// Copyright Epic Games, Inc. All Rights Reserved.
+
+#pragma once
+
+#include "CoreMinimal.h"
+#include "UObject/ObjectMacros.h"
+#include "UObject/Class.h"
+#include "GameFramework/Actor.h"
+#include "GameplayDebuggerTypes.h"
+#include "GameplayDebuggerCategoryReplicator.generated.h"
+
+class AGameplayDebuggerCategoryReplicator;
+class APlayerController;
+class FGameplayDebuggerCategory;
+class FGameplayDebuggerExtension;
+class UGameplayDebuggerRenderingComponent;
+
+// Struct used to send the DataPackPackets as RPC`s instead of using the CustomDeltaNetSerializer.
+USTRUCT()
+struct FGameplayDebuggerDataPackRPCParams
+{
+	GENERATED_BODY()
+
+	UPROPERTY()
+	FName CategoryName;
+
+	UPROPERTY()
+	int32 DataPackIdx = 0;
+
+	UPROPERTY()
+	FGameplayDebuggerDataPackHeader Header;
+
+	UPROPERTY()
+	TArray<uint8> Data;
+};
+
+USTRUCT()
+struct FGameplayDebuggerCategoryData
+{
+	GENERATED_BODY()
+
+	UPROPERTY()
+	FName CategoryName;
+
+	UPROPERTY()
+	TArray<FString> TextLines;
+
+	UPROPERTY()
+	TArray<FGameplayDebuggerShape> Shapes;
+
+	// Either replicated using the NetDeltaSerialize or alternatively as regular RPC`s
+	UPROPERTY(NotReplicated)
+	TArray<FGameplayDebuggerDataPackHeader> DataPacks;
+
+	UPROPERTY()
+	bool bIsEnabled = false;
+};
+
+USTRUCT()
+struct FGameplayDebuggerNetPack
+{
+	GENERATED_BODY()
+
+	UPROPERTY(NotReplicated)
+	TObjectPtr<AGameplayDebuggerCategoryReplicator> Owner;
+
+	FGameplayDebuggerNetPack() : Owner(nullptr) {}
+
+	bool NetDeltaSerialize(FNetDeltaSerializeInfo & DeltaParms);
+	void OnCategoriesChanged();
+
+#if UE_WITH_IRIS
+	// When using Iris we cannot use the same appoach as the current NetSerializeDelta as it does polling of data and also modifies data outside of the replicated state during serialization which is not allowed when using iris.
+	// Therefore we need explicit methods to populate the state from the owner and to apply it to the owner which can be invoked before and after serialization as needed.
+	void PopulateFromOwner();
+	void ApplyToOwner();
+#endif
+
+private:
+	UPROPERTY()
+	TArray<FGameplayDebuggerCategoryData> SavedData;
+};
+
+template<>
+struct TStructOpsTypeTraits<FGameplayDebuggerNetPack> : public TStructOpsTypeTraitsBase2<FGameplayDebuggerNetPack>
+{
+	enum
+	{
+		WithNetDeltaSerializer = true,
+	};
+};
+
+USTRUCT()
+struct FGameplayDebuggerDebugActor
+{
+	GENERATED_USTRUCT_BODY()
+
+	UPROPERTY(NotReplicated)
+	TWeakObjectPtr<AActor> Actor;
+
+	UPROPERTY()
+	FName ActorName;
+
+	UPROPERTY()
+	int32 SyncCounter = 0;
+};
+
+USTRUCT()
+struct FGameplayDebuggerVisLogSync
+{
+	GENERATED_USTRUCT_BODY()
+
+	UPROPERTY()
+	FString DeviceIDs;
+};
+
+DECLARE_MULTICAST_DELEGATE_TwoParams(FNotifyGameplayDebuggerOwnerChange, AGameplayDebuggerCategoryReplicator*, APlayerController* /** Old Owner */);
+
+UCLASS(NotBlueprintable, NotBlueprintType, notplaceable, noteditinlinenew, hidedropdown, Transient)
+class GAMEPLAYDEBUGGER_API AGameplayDebuggerCategoryReplicator : public AActor
+{
+	GENERATED_UCLASS_BODY()
+
+	virtual class UNetConnection* GetNetConnection() const override;
+	virtual bool IsNetRelevantFor(const AActor* RealViewer, const AActor* ViewTarget, const FVector& SrcLocation) const override;
+
+protected:
+	virtual void BeginPlay() override;
+
+#if UE_WITH_IRIS
+	virtual void BeginReplication() override;
+
+	virtual const AActor* GetNetOwner() const override;
+#endif
+
+public:
+	virtual void EndPlay(const EEndPlayReason::Type EndPlayReason) override;
+	virtual void TickActor(float DeltaTime, enum ELevelTick TickType, FActorTickFunction& ThisTickFunction) override;
+	virtual void PostNetReceive() override;
+
+	virtual void PreReplication(IRepChangedPropertyTracker& ChangedPropertyTracker) override;
+
+	/** [AUTH] set new owner */
+	void SetReplicatorOwner(APlayerController* InOwnerPC);
+	static FNotifyGameplayDebuggerOwnerChange NotifyDebuggerOwnerChange;
+
+	/** [ALL] set replicator state */
+	void SetEnabled(bool bEnable);
+
+	/** [ALL] set category state */
+	void SetCategoryEnabled(int32 CategoryId, bool bEnable);
+
+	/** [ALL] set actor for debugging */
+	void SetDebugActor(AActor* Actor, bool bSelectInEditor = false);
+
+	/** [ALL] set view matrix that should be used for culling and picking instead of the associated controller view point */
+	void SetViewPoint(const FVector& InViewLocation, const FVector& InViewDirection);
+
+	/** Indicates that a view matrix has been set for culling and picking instead of the associated controller view point */
+	bool IsViewPointSet() const { return ViewLocation.IsSet() && ViewDirection.IsSet(); }
+
+	/** [ALL] reset view matrix that should be used for culling and picking to revert on the associated controller view point */
+	void ResetViewPoint();
+
+	/** [ALL] send input event to category */
+	void SendCategoryInputEvent(int32 CategoryId, int32 HandlerId);
+
+	/** [ALL] send input event to extension */
+	void SendExtensionInputEvent(int32 ExtensionId, int32 HandlerId);
+
+	/** [AUTH] starts data collection */
+	void CollectCategoryData(bool bForce = false);
+
+	/** get current debug actor */
+	AActor* GetDebugActor() const { return DebugActor.Actor.Get(); }
+	
+	/** get name of debug actor */
+	FName GetDebugActorName() const { return DebugActor.ActorName; }
+
+	/** get sync counter, increased with every change of DebugActor */
+	int32 GetDebugActorCounter() const { return DebugActor.SyncCounter; }
+
+	/** get view point information */
+	bool GetViewPoint(FVector& OutViewLocation, FVector& OutViewDirection) const;
+
+	const FGameplayDebuggerVisLogSync& GetVisLogSyncData() const { return VisLogSync; }
+
+	/** get player controller owning this replicator */
+	APlayerController* GetReplicationOwner() const { return OwnerPC; }
+
+	/** get replicator state */
+	bool IsEnabled() const { return bIsEnabled; }
+
+	/** returns true if this replicator has been created for an editor world */
+	bool IsEditorWorldReplicator() const { return bIsEditorWorldReplicator; }
+
+	/** get category state */
+	bool IsCategoryEnabled(int32 CategoryId) const;
+
+	/** check if debug actor was selected */
+	bool HasDebugActor() const { return DebugActor.ActorName != NAME_None; }
+
+	/** get category count */
+	int32 GetNumCategories() const { return Categories.Num(); }
+
+	/** get extension count */
+	int32 GetNumExtensions() const { return Extensions.Num(); }
+
+	/** get category object */
+	TSharedRef<FGameplayDebuggerCategory> GetCategory(int32 CategoryId) const { return Categories[CategoryId]; }
+
+	/** get category object */
+	TSharedRef<FGameplayDebuggerExtension> GetExtension(int32 ExtensionId) const { return Extensions[ExtensionId]; }
+
+	/** returns true if object was created for local player (client / standalone) */
+	bool IsLocal() const { return bIsLocal; }
+
+#if WITH_EDITOR
+	void InitForEditor();
+#endif // WITH_EDITOR
+
+protected:
+	UFUNCTION()
+	void OnRep_ReplicatedData();
+
+
+	friend FGameplayDebuggerNetPack;
+
+	UPROPERTY(Replicated)
+	TObjectPtr<APlayerController> OwnerPC;
+
+	UPROPERTY(Replicated)
+	bool bIsEnabled;
+
+	UPROPERTY(Replicated, ReplicatedUsing=OnRep_ReplicatedData)
+	FGameplayDebuggerNetPack ReplicatedData;
+
+	UPROPERTY(Replicated)
+	FGameplayDebuggerDebugActor	DebugActor;
+
+	UPROPERTY(Replicated)
+	FGameplayDebuggerVisLogSync	VisLogSync;
+
+	/** rendering component needs to attached to some actor, and this is as good as any */
+	UPROPERTY()
+	TObjectPtr<UGameplayDebuggerRenderingComponent> RenderingComp;
+
+	/** category objects */
+	TArray<TSharedRef<FGameplayDebuggerCategory> > Categories;
+
+	/** extension objects */
+	TArray<TSharedRef<FGameplayDebuggerExtension> > Extensions;
+
+	TOptional<FVector> ViewLocation;
+	TOptional<FVector> ViewDirection;
+
+	uint32 bIsEnabledLocal : 1;
+	uint32 bHasAuthority : 1;
+	uint32 bIsLocal : 1;
+	uint32 bIsEditorWorldReplicator : 1;
+	uint32 bSendDataPacksUsingRPC : 1;
+
+	/** notify about changes in known category set */
+	void OnCategoriesChanged();
+
+	/** notify about changes in known extensions set */
+	void OnExtensionsChanged();
+
+	/** send notifies to all categories about current tool state */
+	void NotifyCategoriesToolState(bool bIsActive);
+
+	/** send notifies to all categories about current tool state */
+	void NotifyExtensionsToolState(bool bIsActive);
+
+	UFUNCTION(Server, Reliable, WithValidation, meta = (CallInEditor = "true"))
+	void ServerSetEnabled(bool bEnable);
+
+	UFUNCTION(Server, Reliable, WithValidation, meta = (CallInEditor = "true"))
+	void ServerSetDebugActor(AActor* Actor, bool bSelectInEditor);
+
+	UFUNCTION(Server, Reliable, WithValidation, meta = (CallInEditor = "true"))
+	void ServerSetViewPoint(const FVector& InViewLocation, const FVector& InViewDirection);
+
+	UFUNCTION(Server, Reliable, WithValidation, meta = (CallInEditor = "true"))
+	void ServerResetViewPoint();
+
+	UFUNCTION(Server, Reliable, WithValidation, meta = (CallInEditor = "true"))
+	void ServerSetCategoryEnabled(int32 CategoryId, bool bEnable);
+
+	/** helper function for replicating input for category handlers */
+	UFUNCTION(Server, Reliable, WithValidation, meta = (CallInEditor = "true"))
+	void ServerSendCategoryInputEvent(int32 CategoryId, int32 HandlerId);
+
+	/** helper function for replicating input for extension handlers */
+	UFUNCTION(Server, Reliable, WithValidation, meta = (CallInEditor = "true"))
+	void ServerSendExtensionInputEvent(int32 ExtensionId, int32 HandlerId);
+
+	/** helper function for optionally sending DataPackPackets as RPC's */
+	UFUNCTION(Client, Reliable)
+	void ClientDataPackPacket(const FGameplayDebuggerDataPackRPCParams& Params);
+	
+	/** [LOCAL] notify from CategoryData replication */
+	void OnReceivedDataPackPacket(int32 CategoryId, int32 DataPackId, const FGameplayDebuggerDataPack& DataPacket);
+
+	/** called both from BeginPlay and InitForEditor to setup instance's internal */
+	void Init();
+
+private:
+	void SendDataPackPacket(FName CategoryName, int32 DataPackIdx, FGameplayDebuggerDataPack& DataPack);
+};
diff --git a/Engine/Source/Runtime/GameplayDebugger/Public/GameplayDebuggerCompat.h b/Engine/Source/Runtime/GameplayDebugger/Public/GameplayDebuggerCompat.h
new file mode 100644
--- /dev/null
+++ b/Engine/Source/Runtime/GameplayDebugger/Public/GameplayDebuggerCompat.h
@@ -0,0 +1,72 @@
+// Copyright Epic Games, Inc. All Rights Reserved.
+
+#pragma once
+
+#if WITH_EDITOR
+#include "Editor/EditorEngine.h"
+#include "LevelEditor.h"
+#endif // WITH_EDITOR
+#include "GameplayDebuggerAddonManager.h"
+#include "GameplayDebuggerPlayerManager.h"
+
+class AGameplayDebuggingReplicator;
+
+class FGameplayDebuggerCompat : public FSelfRegisteringExec, public GameplayDebugger
+{
+public:
+	// Begin IModuleInterface
+	virtual void StartupModule() override;
+	virtual void ShutdownModule() override;
+	// End IModuleInterface
+
+	void WorldAdded(UWorld* InWorld);
+	void WorldDestroyed(UWorld* InWorld);
+#if WITH_EDITOR
+	void OnLevelActorAdded(AActor* InActor);
+	void OnLevelActorDeleted(AActor* InActor);
+	TSharedRef<FExtender> OnExtendLevelEditorViewMenu(const TSharedRef<FUICommandList> CommandList);
+	void CreateSnappingOptionsMenu(FMenuBuilder& Builder);
+	void CreateSettingSubMenu(FMenuBuilder& Builder);
+	void HandleSettingChanged(FName PropertyName);
+#endif
+
+	TArray<TWeakObjectPtr<AGameplayDebuggingReplicator> >& GetAllReplicators(UWorld* InWorld);
+	void AddReplicator(UWorld* InWorld, AGameplayDebuggingReplicator* InReplicator);
+	void RemoveReplicator(UWorld* InWorld, AGameplayDebuggingReplicator* InReplicator);
+
+	// Begin FExec Interface
+	virtual bool Exec(UWorld* Inworld, const TCHAR* Cmd, FOutputDevice& Ar) override;
+	// End FExec Interface
+
+private:
+	virtual bool CreateGameplayDebuggerForPlayerController(APlayerController* PlayerController) override;
+	virtual bool IsGameplayDebuggerActiveForPlayerController(APlayerController* PlayerController) override;
+
+	bool DoesGameplayDebuggingReplicatorExistForPlayerController(APlayerController* PlayerController);
+
+	TMap<TWeakObjectPtr<UWorld>, TArray<TWeakObjectPtr<AGameplayDebuggingReplicator> > > AllReplicatorsPerWorlds;
+
+public:
+	virtual void RegisterCategory(FName CategoryName, FOnGetCategory MakeInstanceDelegate, EGameplayDebuggerCategoryState CategoryState = EGameplayDebuggerCategoryState::Disabled, int32 SlotIdx = INDEX_NONE);
+	virtual void UnregisterCategory(FName CategoryName);
+	virtual void NotifyCategoriesChanged();
+	virtual void RegisterExtension(FName ExtensionName, IGameplayDebugger::FOnGetExtension MakeInstanceDelegate);
+	virtual void UnregisterExtension(FName ExtensionName);
+	virtual void NotifyExtensionsChanged();
+	virtual void UseNewGameplayDebugger();
+
+	void StartupNewDebugger();
+	void ShutdownNewDebugger();
+	AGameplayDebuggerPlayerManager& GetPlayerManager(UWorld* World);
+	void OnWorldInitialized(UWorld* World, const UWorld::InitializationValues IVS);
+
+	bool bNewDebuggerEnabled;
+	FGameplayDebuggerAddonManager AddonManager;
+	TMap<TWeakObjectPtr<UWorld>, TWeakObjectPtr<AGameplayDebuggerPlayerManager>> PlayerManagers;
+
+private:
+
+#if WITH_EDITOR
+	FLevelEditorModule::FLevelEditorMenuExtender ViewMenuExtender;
+#endif
+};
diff --git a/Engine/Source/Runtime/GameplayDebugger/Public/GameplayDebuggerConfig.h b/Engine/Source/Runtime/GameplayDebugger/Public/GameplayDebuggerConfig.h
new file mode 100644
--- /dev/null
+++ b/Engine/Source/Runtime/GameplayDebugger/Public/GameplayDebuggerConfig.h
@@ -0,0 +1,259 @@
+// Copyright Epic Games, Inc. All Rights Reserved.
+
+#pragma once
+
+#include "Containers/Array.h"
+#include "Containers/Map.h"
+#include "Containers/UnrealString.h"
+#include "CoreMinimal.h"
+#include "Engine/DeveloperSettings.h"
+#include "HAL/Platform.h"
+#include "InputCoreTypes.h"
+#include "UObject/NameTypes.h"
+#include "UObject/Object.h"
+#include "UObject/ObjectMacros.h"
+#include "UObject/UObjectGlobals.h"
+
+#include "GameplayDebuggerConfig.generated.h"
+
+class FArchive;
+struct FGameplayDebuggerInputModifier;
+
+UENUM()
+enum class EGameplayDebuggerOverrideMode : uint8
+{
+	Enable,
+	Disable,
+	UseDefault,
+};
+
+USTRUCT()
+struct FGameplayDebuggerInputConfig
+{
+	GENERATED_USTRUCT_BODY()
+
+	FGameplayDebuggerInputConfig()
+		: bModShift(false)
+		, bModCtrl(false)
+		, bModAlt(false)
+		, bModCmd(false) {}
+
+	UPROPERTY(VisibleAnywhere, Category = Input)
+	FString ConfigName;
+
+	UPROPERTY(EditAnywhere, Category = Input)
+	FKey Key;
+
+	UPROPERTY(EditAnywhere, Category = Input)
+	uint32 bModShift : 1;
+
+	UPROPERTY(EditAnywhere, Category = Input)
+	uint32 bModCtrl : 1;
+
+	UPROPERTY(EditAnywhere, Category = Input)
+	uint32 bModAlt : 1;
+
+	UPROPERTY(EditAnywhere, Category = Input)
+	uint32 bModCmd : 1;
+};
+
+USTRUCT()
+struct FGameplayDebuggerCategoryConfig
+{
+	GENERATED_USTRUCT_BODY()
+
+	UPROPERTY(VisibleAnywhere, Category = Settings)
+	FString CategoryName;
+
+	UPROPERTY(EditAnywhere, Category = Settings, meta = (EditCondition = "bOverrideSlotIdx", ClampMin = -1, ClampMax = 9, UIMin = -1, UIMax = 9))
+	int32 SlotIdx;
+
+	UPROPERTY(EditAnywhere, Category = Settings)
+	EGameplayDebuggerOverrideMode ActiveInGame;
+
+	UPROPERTY(EditAnywhere, Category = Settings)
+	EGameplayDebuggerOverrideMode ActiveInSimulate;
+
+	UPROPERTY(EditAnywhere, Category = Settings)
+	EGameplayDebuggerOverrideMode Hidden;
+
+	UPROPERTY(EditAnywhere, Category = Settings, AdvancedDisplay)
+	uint32 bOverrideSlotIdx : 1;
+
+	UPROPERTY(EditAnywhere, Category = Settings, EditFixedSize)
+	TArray<FGameplayDebuggerInputConfig> InputHandlers;
+
+	FGameplayDebuggerCategoryConfig()
+		: SlotIdx(0)
+		, ActiveInGame(EGameplayDebuggerOverrideMode::UseDefault)
+		, ActiveInSimulate(EGameplayDebuggerOverrideMode::UseDefault)
+		, Hidden(EGameplayDebuggerOverrideMode::UseDefault)
+		, bOverrideSlotIdx(false)
+	{}
+};
+
+USTRUCT()
+struct FGameplayDebuggerExtensionConfig
+{
+	GENERATED_USTRUCT_BODY()
+
+	UPROPERTY(VisibleAnywhere, Category = Settings)
+	FString ExtensionName;
+
+	UPROPERTY(EditAnywhere, Category = Settings)
+	EGameplayDebuggerOverrideMode UseExtension;
+
+	UPROPERTY(EditAnywhere, Category = Settings, EditFixedSize)
+	TArray<FGameplayDebuggerInputConfig> InputHandlers;
+
+	FGameplayDebuggerExtensionConfig() : UseExtension(EGameplayDebuggerOverrideMode::UseDefault) {}
+};
+
+UCLASS(config = Engine, defaultconfig)
+class GAMEPLAYDEBUGGER_API UGameplayDebuggerConfig : public UObject
+{
+	GENERATED_UCLASS_BODY()
+
+	/** key used to activate visual debugger tool */
+	UPROPERTY(config, EditAnywhere, Category = Input)
+	FKey ActivationKey;
+
+	/** select next category row */
+	UPROPERTY(config, EditAnywhere, Category = Input)
+	FKey CategoryRowNextKey;
+
+	/** select previous category row */
+	UPROPERTY(config, EditAnywhere, Category = Input)
+	FKey CategoryRowPrevKey;
+
+	/** select category slot 0 */
+	UPROPERTY(config, EditAnywhere, Category = Input)
+	FKey CategorySlot0;
+
+	/** select category slot 1 */
+	UPROPERTY(config, EditAnywhere, Category = Input)
+	FKey CategorySlot1;
+
+	/** select category slot 2 */
+	UPROPERTY(config, EditAnywhere, Category = Input)
+	FKey CategorySlot2;
+
+	/** select category slot 3 */
+	UPROPERTY(config, EditAnywhere, Category = Input)
+	FKey CategorySlot3;
+
+	/** select category slot 4 */
+	UPROPERTY(config, EditAnywhere, Category = Input)
+	FKey CategorySlot4;
+
+	/** select category slot 5 */
+	UPROPERTY(config, EditAnywhere, Category = Input)
+	FKey CategorySlot5;
+
+	/** select category slot 6 */
+	UPROPERTY(config, EditAnywhere, Category = Input)
+	FKey CategorySlot6;
+
+	/** select category slot 7 */
+	UPROPERTY(config, EditAnywhere, Category = Input)
+	FKey CategorySlot7;
+
+	/** select category slot 8 */
+	UPROPERTY(config, EditAnywhere, Category = Input)
+	FKey CategorySlot8;
+
+	/** select category slot 9 */
+	UPROPERTY(config, EditAnywhere, Category = Input)
+	FKey CategorySlot9;
+
+	/** additional canvas padding: left */
+	UPROPERTY(config, EditAnywhere, Category = Display)
+	float DebugCanvasPaddingLeft;
+
+	/** additional canvas padding: right */
+	UPROPERTY(config, EditAnywhere, Category = Display)
+	float DebugCanvasPaddingRight;
+
+	/** additional canvas padding: top */
+	UPROPERTY(config, EditAnywhere, Category = Display)
+	float DebugCanvasPaddingTop;
+
+	/** additional canvas padding: bottom */
+	UPROPERTY(config, EditAnywhere, Category = Display)
+	float DebugCanvasPaddingBottom;
+
+	/** enable text shadow by default */
+	UPROPERTY(config, EditAnywhere, Category = Display)
+	bool bDebugCanvasEnableTextShadow;
+
+	UPROPERTY(config, EditAnywhere, Category = AddOns, EditFixedSize)
+	TArray<FGameplayDebuggerCategoryConfig> Categories;
+
+	UPROPERTY(config, EditAnywhere, Category = AddOns, EditFixedSize)
+	TArray<FGameplayDebuggerExtensionConfig> Extensions;
+
+	/** updates entry in Categories array and modifies category creation params */
+	void UpdateCategoryConfig(const FName CategoryName, int32& SlotIdx, uint8& CategoryState);
+
+	/** updates entry in Categories array and modifies input binding params */
+	void UpdateCategoryInputConfig(const FName CategoryName, const FName InputName, FName& KeyName, FGameplayDebuggerInputModifier& KeyModifier);
+
+	/** updates entry in Extensions array and modifies extension creation params */
+	void UpdateExtensionConfig(const FName ExtensionName, uint8& UseExtension);
+
+	/** updates entry in Categories array and modifies input binding params */
+	void UpdateExtensionInputConfig(const FName ExtensionName, const FName InputName, FName& KeyName, FGameplayDebuggerInputModifier& KeyModifier);
+
+	/** remove all category and extension data from unknown sources (outdated entries) */
+	void RemoveUnknownConfigs();
+
+	virtual void Serialize(FArchive& Ar) override;
+#if WITH_EDITOR
+	virtual void PostEditChangeChainProperty(struct FPropertyChangedChainEvent& PropertyChangedEvent) override;
+#endif
+
+private:
+
+	/** used for cleanup */
+	TArray<FName> KnownCategoryNames;
+	TArray<FName> KnownExtensionNames;
+	TMultiMap<FName, FName> KnownCategoryInputNames;
+	TMultiMap<FName, FName> KnownExtensionInputNames;
+};
+
+
+UCLASS(config = EditorPerProjectUserSettings, meta = (DisplayName = "Gameplay Debugger"))
+class GAMEPLAYDEBUGGER_API UGameplayDebuggerUserSettings : public UDeveloperSettings
+{
+	GENERATED_BODY()
+protected:
+	virtual FName GetCategoryName() const override { return TEXT("Advanced"); }
+
+public:
+	static int32 GetFontSize() { return GetDefault<UGameplayDebuggerUserSettings>()->FontSize; }
+	static void SetFontSize(const int32 InFontSize);
+
+	/** Controls whether GameplayDebugger will be available in pure editor mode.
+	 *  @Note that you need to reload the map for the changes to this property to take effect */
+	UPROPERTY(config, EditAnywhere, Category = GameplayDebugger)
+	uint32 bEnableGameplayDebuggerInEditor : 1;
+
+	/**
+	 * Distance from view location under which actors can be selected
+	 * This distance can also be used by some categories to apply culling.
+	 */
+	UPROPERTY(config, EditAnywhere, Category = GameplayDebugger)
+	float MaxViewDistance = 25000.0f;
+
+	/**
+	 * Angle from view direction under which actors can be selected
+	 * This angle can also be used by some categories to apply culling.
+	 */
+	UPROPERTY(config, EditAnywhere, Category = GameplayDebugger, meta = (UIMin = 0, ClampMin = 0, UIMax = 180, ClampMax = 180, Units = deg))
+	float MaxViewAngle = 45.f;
+
+protected:
+	/** Font Size used by Gameplay Debugger */ 
+	UPROPERTY(config, EditAnywhere, Category = GameplayDebugger)
+	int32 FontSize = 10;
+};
diff --git a/Engine/Source/Runtime/GameplayDebugger/Public/GameplayDebuggerExtension.h b/Engine/Source/Runtime/GameplayDebugger/Public/GameplayDebuggerExtension.h
new file mode 100644
--- /dev/null
+++ b/Engine/Source/Runtime/GameplayDebugger/Public/GameplayDebuggerExtension.h
@@ -0,0 +1,49 @@
+// Copyright Epic Games, Inc. All Rights Reserved.
+
+// GAMEPLAY DEBUGGER EXTENSION
+// 
+// Extensions allows creating additional key bindings for gameplay debugger.
+// For example, you can use them to add another way of selecting actor to Debug.
+//
+// Replication is limited only to handling input events and tool state events,
+// it's not possible to send variables or RPC calls
+//
+// It should be compiled and used only when module is included, so every extension class
+// needs be placed in #if WITH_GAMEPLAY_DEBUGGER block.
+// 
+// Extensions needs to be manually registered and unregistered with GameplayDebugger.
+// It's best to do it in owning module's Startup / Shutdown, similar to detail view customizations.
+
+#pragma once
+
+#include "Containers/UnrealString.h"
+#include "CoreMinimal.h"
+#include "GameplayDebuggerAddonBase.h"
+
+class APlayerController;
+
+class GAMEPLAYDEBUGGER_API FGameplayDebuggerExtension : public FGameplayDebuggerAddonBase
+{
+public:
+
+	virtual ~FGameplayDebuggerExtension() {}
+	virtual void OnGameplayDebuggerActivated() override;
+	virtual void OnGameplayDebuggerDeactivated() override;
+
+	/** [LOCAL] description for gameplay debugger's header row, newline character is ignored */
+	virtual FString GetDescription() const;
+
+	/** [LOCAL] called when added to debugger tool or tool is activated */
+	virtual void OnActivated();
+
+	/** [LOCAL] called when removed from debugger tool or tool is deactivated */
+	virtual void OnDeactivated();
+
+	/** check if extension is created for local player */
+	bool IsLocal() const;
+
+protected:
+
+	/** get player controller owning gameplay debugger tool */
+	APlayerController* GetPlayerController() const;
+};
diff --git a/Engine/Source/Runtime/GameplayDebugger/Public/GameplayDebuggerExtension_HUD.h b/Engine/Source/Runtime/GameplayDebugger/Public/GameplayDebuggerExtension_HUD.h
new file mode 100644
--- /dev/null
+++ b/Engine/Source/Runtime/GameplayDebugger/Public/GameplayDebuggerExtension_HUD.h
@@ -0,0 +1,37 @@
+// Copyright Epic Games, Inc. All Rights Reserved.
+
+#pragma once
+
+#include "CoreMinimal.h"
+#include "GameplayDebuggerExtension.h"
+
+class FGameplayDebuggerExtension_HUD : public FGameplayDebuggerExtension
+{
+public:
+
+	FGameplayDebuggerExtension_HUD();
+
+	virtual void OnActivated() override;
+	virtual void OnDeactivated() override;
+	virtual FString GetDescription() const override;
+
+	static TSharedRef<FGameplayDebuggerExtension> MakeInstance();
+
+protected:
+
+	void ToggleGameHUD();
+	void ToggleDebugMessages();
+	void SetGameHUDEnabled(bool bEnable);
+	void SetDebugMessagesEnabled(bool bEnable);
+
+	uint32 bWantsHUDEnabled : 1;
+	uint32 bIsGameHUDEnabled : 1;
+	uint32 bAreDebugMessagesEnabled : 1;
+	uint32 bPrevDebugMessagesEnabled : 1;
+	mutable uint32 bIsCachedDescriptionValid : 1;
+
+	int32 HudBindingIdx;
+	int32 MessagesBindingIdx;
+
+	mutable FString CachedDescription;
+};
diff --git a/Engine/Source/Runtime/GameplayDebugger/Public/GameplayDebuggerExtension_Spectator.h b/Engine/Source/Runtime/GameplayDebugger/Public/GameplayDebuggerExtension_Spectator.h
new file mode 100644
--- /dev/null
+++ b/Engine/Source/Runtime/GameplayDebugger/Public/GameplayDebuggerExtension_Spectator.h
@@ -0,0 +1,29 @@
+// Copyright Epic Games, Inc. All Rights Reserved.
+
+#pragma once
+
+#include "CoreMinimal.h"
+#include "GameplayDebuggerExtension.h"
+
+class ADebugCameraController;
+
+class FGameplayDebuggerExtension_Spectator : public FGameplayDebuggerExtension
+{
+public:
+
+	FGameplayDebuggerExtension_Spectator();
+
+	virtual void OnDeactivated() override;
+	virtual FString GetDescription() const override;
+
+	static TSharedRef<FGameplayDebuggerExtension> MakeInstance();
+
+protected:
+
+	void ToggleSpectatorMode();
+
+	uint32 bHasInputBinding : 1;
+	mutable uint32 bIsCachedDescriptionValid : 1;
+	mutable FString CachedDescription;
+	TWeakObjectPtr<ADebugCameraController> SpectatorController;
+};
diff --git a/Engine/Source/Runtime/GameplayDebugger/Public/GameplayDebuggerLocalController.h b/Engine/Source/Runtime/GameplayDebugger/Public/GameplayDebuggerLocalController.h
new file mode 100644
--- /dev/null
+++ b/Engine/Source/Runtime/GameplayDebugger/Public/GameplayDebuggerLocalController.h
@@ -0,0 +1,147 @@
+// Copyright Epic Games, Inc. All Rights Reserved.
+
+#pragma once
+
+#include "CoreMinimal.h"
+#include "UObject/ObjectMacros.h"
+#include "UObject/Object.h"
+#include "Engine/EngineTypes.h"
+#include "GameplayDebuggerLocalController.generated.h"
+
+class AActor;
+class AGameplayDebuggerCategoryReplicator;
+class AGameplayDebuggerPlayerManager;
+class FGameplayDebuggerCanvasContext;
+class FGameplayDebuggerCategory;
+class UInputComponent;
+struct FKey;
+class UFont;
+
+
+UCLASS(NotBlueprintable, NotBlueprintType, noteditinlinenew, hidedropdown, Transient)
+class UGameplayDebuggerLocalController : public UObject
+{
+	GENERATED_UCLASS_BODY()
+
+	virtual void BeginDestroy() override;
+
+	/** initialize controller with replicator owner */
+	void Initialize(AGameplayDebuggerCategoryReplicator& Replicator, AGameplayDebuggerPlayerManager& Manager);
+
+	/** remove from world */
+	void Cleanup();
+
+	/** drawing event */
+	void OnDebugDraw(class UCanvas* Canvas, class APlayerController* PC);
+
+	/** binds input actions */
+	void BindInput(UInputComponent& InputComponent);
+
+	/** checks if key is bound by any action */
+	bool GAMEPLAYDEBUGGER_API IsKeyBound(const FName KeyName) const;
+
+protected:
+	friend struct FGameplayDebuggerConsoleCommands;
+
+	UPROPERTY()
+	TObjectPtr<AGameplayDebuggerCategoryReplicator> CachedReplicator;
+
+	UPROPERTY()
+	TObjectPtr<AGameplayDebuggerPlayerManager> CachedPlayerManager;
+
+	UPROPERTY()
+	TObjectPtr<AActor> DebugActorCandidate;
+
+	UPROPERTY()
+	TObjectPtr<UFont> HUDFont;
+
+	TArray<TArray<int32> > DataPackMap;
+	TArray<TArray<int32> > SlotCategoryIds;
+	TArray<FString> SlotNames;
+
+	TSet<FName> UsedBindings;
+
+	uint32 bSimulateMode : 1;
+	uint32 bNeedsCleanup : 1;
+	uint32 bIsSelectingActor : 1;
+	uint32 bIsLocallyEnabled : 1;
+	uint32 bPrevLocallyEnabled : 1;
+	uint32 bEnableTextShadow : 1;
+	uint32 bPrevScreenMessagesEnabled : 1;
+#if WITH_EDITOR
+	uint32 bActivateOnPIEEnd : 1;
+#endif // WITH_EDITOR
+
+	FString ActivationKeyDesc;
+	FString RowUpKeyDesc;
+	FString RowDownKeyDesc;
+	FString CategoryKeysDesc;
+
+	int32 ActiveRowIdx;
+	int32 NumCategorySlots;
+	int32 NumCategories;
+	static constexpr int32 NumCategoriesPerRow = 10;
+
+	float PaddingLeft;
+	float PaddingRight;
+	float PaddingTop;
+	float PaddingBottom;
+
+	FTimerHandle StartSelectingActorHandle;
+	FTimerHandle SelectActorTickHandle;
+
+	void OnActivationPressed();
+	void OnActivationReleased();
+	void OnCategory0Pressed();
+	void OnCategory1Pressed();
+	void OnCategory2Pressed();
+	void OnCategory3Pressed();
+	void OnCategory4Pressed();
+	void OnCategory5Pressed();
+	void OnCategory6Pressed();
+	void OnCategory7Pressed();
+	void OnCategory8Pressed();
+	void OnCategory9Pressed();
+	void OnCategoryRowUpPressed();
+	void OnCategoryRowDownPressed();
+	void OnCategoryBindingEvent(int32 CategoryId, int32 HandlerId);
+	void OnExtensionBindingEvent(int32 ExtensionId, int32 HandlerId);
+
+	/** sets the local player as the new debug actor */
+	void OnSelectLocalPlayer();
+
+	/** called short time after activation key was pressed and hold */
+	void OnStartSelectingActor();
+
+	/** called in tick during actor selection */
+	void OnSelectActorTick();
+
+	/** toggle state of categories in given slot */
+	void ToggleSlotState(int32 SlotIdx);
+
+	/** toggle debugger on/off */
+	void ToggleActivation();
+
+	/** draw header row */
+	void DrawHeader(FGameplayDebuggerCanvasContext& CanvasContext);
+
+	/** draw header for category */
+	void DrawCategoryHeader(int32 CategoryId, TSharedRef<FGameplayDebuggerCategory> Category, FGameplayDebuggerCanvasContext& CanvasContext);
+
+#if WITH_EDITOR
+	/** event for simulate in editor mode */
+	void OnSelectionChanged(UObject* Object);
+	void OnSelectedObject(UObject* Object);
+	void OnBeginPIE(const bool bIsSimulating);
+	void OnEndPIE(const bool bIsSimulating);
+#endif
+
+	FString GetKeyDescriptionShort(const FKey& KeyBind) const;
+	FString GetKeyDescriptionLong(const FKey& KeyBind) const;
+
+	/** called when known category set has changed */
+	void OnCategoriesChanged();
+
+	/** build DataPackMap for replication details */
+	void RebuildDataPackMap();
+};
\ No newline at end of file
diff --git a/Engine/Source/Runtime/GameplayDebugger/Public/GameplayDebuggerModule.h b/Engine/Source/Runtime/GameplayDebugger/Public/GameplayDebuggerModule.h
new file mode 100644
--- /dev/null
+++ b/Engine/Source/Runtime/GameplayDebugger/Public/GameplayDebuggerModule.h
@@ -0,0 +1,41 @@
+// Copyright Epic Games, Inc. All Rights Reserved.
+
+#pragma once
+
+#include "CoreMinimal.h"
+#include "Modules/ModuleManager.h"
+#include "Engine/World.h"
+#include "GameplayDebugger.h"
+#include "GameplayDebuggerAddonManager.h"
+
+class AGameplayDebuggerPlayerManager;
+
+DECLARE_MULTICAST_DELEGATE(FOnLocalControllerInitialized)
+DECLARE_MULTICAST_DELEGATE(FOnLocalControllerUninitialized)
+DECLARE_MULTICAST_DELEGATE(FOnDebuggerEdMode)
+
+class FGameplayDebuggerModule : public IGameplayDebugger
+{
+public:	
+	virtual void StartupModule() override;
+	virtual void ShutdownModule() override;
+
+	virtual void RegisterCategory(FName CategoryName, IGameplayDebugger::FOnGetCategory MakeInstanceDelegate, EGameplayDebuggerCategoryState CategoryState, int32 SlotIdx) override;
+	virtual void UnregisterCategory(FName CategoryName) override;
+	virtual void NotifyCategoriesChanged() override;
+	virtual void RegisterExtension(FName ExtensionName, IGameplayDebugger::FOnGetExtension MakeInstanceDelegate) override;
+	virtual void UnregisterExtension(FName ExtensionName) override;
+	virtual void NotifyExtensionsChanged() override;
+
+	AGameplayDebuggerPlayerManager& GetPlayerManager(UWorld* World);
+	void OnWorldInitialized(UWorld* World, const UWorld::InitializationValues IVS);
+		
+	FGameplayDebuggerAddonManager AddonManager;
+	TMap<TWeakObjectPtr<UWorld>, TWeakObjectPtr<AGameplayDebuggerPlayerManager>> PlayerManagers;
+
+	static GAMEPLAYDEBUGGER_API FOnLocalControllerInitialized OnLocalControllerInitialized;
+	static GAMEPLAYDEBUGGER_API FOnLocalControllerUninitialized OnLocalControllerUninitialized;
+#if WITH_EDITOR
+	static GAMEPLAYDEBUGGER_API FOnDebuggerEdMode OnDebuggerEdModeActivation;
+#endif // WITH_EDITOR
+};
diff --git a/Engine/Source/Runtime/GameplayDebugger/Public/GameplayDebuggerPlayerManager.h b/Engine/Source/Runtime/GameplayDebugger/Public/GameplayDebuggerPlayerManager.h
new file mode 100644
--- /dev/null
+++ b/Engine/Source/Runtime/GameplayDebugger/Public/GameplayDebuggerPlayerManager.h
@@ -0,0 +1,97 @@
+// Copyright Epic Games, Inc. All Rights Reserved.
+
+#pragma once
+
+#include "CoreMinimal.h"
+#include "UObject/ObjectMacros.h"
+#include "GameFramework/Actor.h"
+#include "Tickable.h"
+#include "GameplayDebuggerPlayerManager.generated.h"
+
+class AGameplayDebuggerCategoryReplicator;
+class APlayerController;
+class UGameplayDebuggerLocalController;
+class UInputComponent;
+class AGameModeBase;
+
+USTRUCT()
+struct FGameplayDebuggerPlayerData
+{
+	GENERATED_USTRUCT_BODY()
+
+	UPROPERTY()
+	TObjectPtr<UGameplayDebuggerLocalController> Controller = nullptr;
+
+	UPROPERTY()
+	TObjectPtr<UInputComponent> InputComponent = nullptr;
+
+	UPROPERTY()
+	TObjectPtr<AGameplayDebuggerCategoryReplicator> Replicator = nullptr;
+};
+
+UCLASS(NotBlueprintable, NotBlueprintType, notplaceable, noteditinlinenew, hidedropdown, Transient)
+class GAMEPLAYDEBUGGER_API AGameplayDebuggerPlayerManager : public AActor, public FTickableGameObject
+{
+	GENERATED_UCLASS_BODY()
+
+	virtual void TickActor(float DeltaTime, enum ELevelTick TickType, FActorTickFunction& ThisTickFunction) override;
+
+	// we're ticking only the manager only when in editor
+	// FTickableGameObject begin
+	virtual UWorld* GetTickableGameObjectWorld() const override { return GetWorld(); }
+	virtual void Tick(float DeltaTime) override;
+	virtual ETickableTickType GetTickableTickType() const override;
+	virtual TStatId GetStatId() const override;
+	virtual bool IsTickableInEditor() const override { return true; }
+	/** results in this instance not being ticked in any of the game worlds */
+	virtual bool IsTickable() const override { return bEditorTimeTick; }
+	// FTickableGameObject end
+
+	virtual void PostInitProperties() override;
+
+protected:
+	virtual void BeginPlay() override;
+
+	void OnReplayScrubTeardown(UWorld* InWorld);
+	void OnGameModeLogout(AGameModeBase* GameMode, AController* Exiting);
+
+public:
+	virtual void EndPlay(const EEndPlayReason::Type Reason) override;
+	void Init();
+		
+	void UpdateAuthReplicators();
+	void RegisterReplicator(AGameplayDebuggerCategoryReplicator& Replicator);
+	void RefreshInputBindings(AGameplayDebuggerCategoryReplicator& Replicator);
+
+	AGameplayDebuggerCategoryReplicator* GetReplicator(const APlayerController& OwnerPC) const;
+	UInputComponent* GetInputComponent(const APlayerController& OwnerPC) const;
+	UGameplayDebuggerLocalController* GetLocalController(const APlayerController& OwnerPC) const;
+#if WITH_EDITOR
+	UGameplayDebuggerLocalController* GetEditorController() const { return EditorWorldData.Controller; }
+#endif // WITH_EDITOR
+	
+	const FGameplayDebuggerPlayerData* GetPlayerData(const APlayerController& OwnerPC) const;
+
+	static AGameplayDebuggerPlayerManager& GetCurrent(UWorld* World);
+
+	/** extracts view location and direction from a player controller that can be used for picking */
+	static void GetViewPoint(const APlayerController& OwnerPC, FVector& OutViewLocation, FVector& OutViewDirection);
+
+protected:
+
+	UPROPERTY()
+	TArray<FGameplayDebuggerPlayerData> PlayerData;
+
+	UPROPERTY()
+	TArray<TObjectPtr<AGameplayDebuggerCategoryReplicator>> PendingRegistrations;
+
+#if WITH_EDITORONLY_DATA 
+	UPROPERTY()
+	FGameplayDebuggerPlayerData EditorWorldData;
+#endif // WITH_EDITORONLY_DATA
+
+	uint8 bEditorTimeTick : 1;
+	uint32 bHasAuthority : 1;
+	uint32 bIsLocal : 1;
+	uint32 bInitialized : 1;
+};
diff --git a/Engine/Source/Runtime/GameplayDebugger/Public/GameplayDebuggerRenderingComponent.h b/Engine/Source/Runtime/GameplayDebugger/Public/GameplayDebuggerRenderingComponent.h
new file mode 100644
--- /dev/null
+++ b/Engine/Source/Runtime/GameplayDebugger/Public/GameplayDebuggerRenderingComponent.h
@@ -0,0 +1,47 @@
+// Copyright Epic Games, Inc. All Rights Reserved.
+
+#pragma once
+
+#include "Debug/DebugDrawComponent.h"
+#include "DebugRenderSceneProxy.h"
+#include "GameplayDebuggerRenderingComponent.generated.h"
+
+class FGameplayDebuggerCompositeSceneProxy;
+
+class FGameplayDebuggerDebugDrawDelegateHelper : public FDebugDrawDelegateHelper
+{
+	typedef FDebugDrawDelegateHelper Super;
+
+public:
+	virtual ~FGameplayDebuggerDebugDrawDelegateHelper() override
+	{
+		Reset();
+	}
+
+	void Reset();
+
+	void AddDelegateHelper(FDebugDrawDelegateHelper* InDebugDrawDelegateHelper);
+
+protected:
+	virtual void RegisterDebugDrawDelegateInternal() override;
+	virtual void UnregisterDebugDrawDelegate() override;
+
+private:
+	TArray<FDebugDrawDelegateHelper*> DebugDrawDelegateHelpers;
+};
+
+
+
+UCLASS(ClassGroup = Debug, NotBlueprintable, NotBlueprintType, noteditinlinenew, hidedropdown, Transient)
+class UGameplayDebuggerRenderingComponent : public UDebugDrawComponent
+{
+	GENERATED_UCLASS_BODY()
+
+protected:
+	virtual FDebugRenderSceneProxy* CreateDebugSceneProxy() override;
+	virtual FDebugDrawDelegateHelper& GetDebugDrawDelegateHelper() override { return GameplayDebuggerDebugDrawDelegateHelper; }
+	virtual FBoxSphereBounds CalcBounds(const FTransform &LocalToWorld) const override;
+
+private:
+	FGameplayDebuggerDebugDrawDelegateHelper GameplayDebuggerDebugDrawDelegateHelper;
+};
diff --git a/Engine/Source/Runtime/GameplayDebugger/Public/GameplayDebuggerTypes.h b/Engine/Source/Runtime/GameplayDebugger/Public/GameplayDebuggerTypes.h
new file mode 100644
--- /dev/null
+++ b/Engine/Source/Runtime/GameplayDebugger/Public/GameplayDebuggerTypes.h
@@ -0,0 +1,404 @@
+// Copyright Epic Games, Inc. All Rights Reserved.
+
+#pragma once
+
+#include "CoreMinimal.h"
+#include "Engine/EngineTypes.h"
+#include "Templates/IsArrayOrRefOfTypeByPredicate.h"
+#include "Traits/IsCharEncodingCompatibleWith.h"
+#include "UObject/ObjectMacros.h"
+#include "GameplayDebuggerTypes.generated.h"
+
+class FCanvasItem;
+class UCanvas;
+class UFont;
+struct FCanvasIcon;
+class APlayerController;
+class UWorld;
+
+DECLARE_LOG_CATEGORY_EXTERN(LogGameplayDebug, Log, All);
+
+class FCanvasItem;
+struct FCanvasIcon;
+
+namespace FGameplayDebuggerUtils
+{
+	bool GAMEPLAYDEBUGGER_API IsAuthority(UWorld* World);
+}
+
+class GAMEPLAYDEBUGGER_API FGameplayDebuggerCanvasContext
+{
+public:
+	/** canvas used for drawing */
+	TWeakObjectPtr<UCanvas> Canvas;
+
+	/** current text font */
+	TWeakObjectPtr<UFont> Font;
+
+	/** the player controller associated with this debugger context  */
+	TWeakObjectPtr<APlayerController> PlayerController;
+
+	/** the world associated with this debugger context  */
+	TWeakObjectPtr<UWorld> World;
+
+	/** font render data */
+	FFontRenderInfo FontRenderInfo;
+
+	/** position of cursor */
+	float CursorX, CursorY;
+
+	/** default position of cursor */
+	float DefaultX, DefaultY;
+
+	FGameplayDebuggerCanvasContext() {}
+	FGameplayDebuggerCanvasContext(UCanvas* InCanvas, UFont* InFont);
+
+	// print string on canvas
+	void Print(const FString& String);
+	void Print(const FColor& Color, const FString& String);
+	void Print(const FColor& Color, const float Alpha, const FString& String);
+	void PrintAt(float PosX, float PosY, const FString& String);
+	void PrintAt(float PosX, float PosY, const FColor& Color, const FString& String);
+	void PrintAt(float PosX, float PosY, const FColor& Color, const float Alpha, const FString& String);
+
+	// print formatted string on canvas
+	template <typename FmtType, typename... Types>
+	void Printf(const FmtType& Fmt, Types... Args)
+	{
+		static_assert(TIsArrayOrRefOfTypeByPredicate<FmtType, TIsCharEncodingCompatibleWithTCHAR>::Value, "Formatting string must be a TCHAR array.");
+		static_assert(TAnd<TIsValidVariadicFunctionArg<Types>...>::Value, "Invalid argument(s) passed to FGameplayDebuggerCanvasContext::Printf");
+
+		PrintfImpl(FColor::White, 1.0f, (const TCHAR*)Fmt, Args...);
+	}
+	
+	template <typename FmtType, typename... Types>
+	void Printf(const FColor& Color, const FmtType& Fmt, Types... Args)
+	{
+		static_assert(TIsArrayOrRefOfTypeByPredicate<FmtType, TIsCharEncodingCompatibleWithTCHAR>::Value, "Formatting string must be a TCHAR array.");
+		static_assert(TAnd<TIsValidVariadicFunctionArg<Types>...>::Value, "Invalid argument(s) passed to FGameplayDebuggerCanvasContext::Printf");
+
+		PrintfImpl(Color, 1.0f, (const TCHAR*)Fmt, Args...);
+	}
+	
+	template <typename FmtType, typename... Types>
+	void Printf(const FColor& Color, const float Alpha, const FmtType& Fmt, Types... Args)
+	{
+		static_assert(TIsArrayOrRefOfTypeByPredicate<FmtType, TIsCharEncodingCompatibleWithTCHAR>::Value, "Formatting string must be a TCHAR array.");
+		static_assert(TAnd<TIsValidVariadicFunctionArg<Types>...>::Value, "Invalid argument(s) passed to FGameplayDebuggerCanvasContext::Printf");
+
+		PrintfImpl(Color, Alpha, (const TCHAR*)Fmt, Args...);
+	}
+	
+	template <typename FmtType, typename... Types>
+	void PrintfAt(float PosX, float PosY, const FmtType& Fmt, Types... Args)
+	{
+		static_assert(TIsArrayOrRefOfTypeByPredicate<FmtType, TIsCharEncodingCompatibleWithTCHAR>::Value, "Formatting string must be a TCHAR array.");
+		static_assert(TAnd<TIsValidVariadicFunctionArg<Types>...>::Value, "Invalid argument(s) passed to FGameplayDebuggerCanvasContext::PrintfAt");
+
+		PrintfAtImpl(PosX, PosY, FColor::White, 1.0f, (const TCHAR*)Fmt, Args...);
+	}
+	
+	template <typename FmtType, typename... Types>
+	void PrintfAt(float PosX, float PosY, const FColor& Color, const FmtType& Fmt, Types... Args)
+	{
+		static_assert(TIsArrayOrRefOfTypeByPredicate<FmtType, TIsCharEncodingCompatibleWithTCHAR>::Value, "Formatting string must be a TCHAR array.");
+		static_assert(TAnd<TIsValidVariadicFunctionArg<Types>...>::Value, "Invalid argument(s) passed to FGameplayDebuggerCanvasContext::PrintfAt");
+
+		PrintfAtImpl(PosX, PosY, Color, 1.0f, (const TCHAR*)Fmt, Args...);
+	}
+	
+	template <typename FmtType, typename... Types>
+	void PrintfAt(float PosX, float PosY, const FColor& Color, const float Alpha, const FmtType& Fmt, Types... Args)
+	{
+		static_assert(TIsArrayOrRefOfTypeByPredicate<FmtType, TIsCharEncodingCompatibleWithTCHAR>::Value, "Formatting string must be a TCHAR array.");
+		static_assert(TAnd<TIsValidVariadicFunctionArg<Types>...>::Value, "Invalid argument(s) passed to FGameplayDebuggerCanvasContext::PrintfAt");
+
+		PrintfAtImpl(PosX, PosY, Color, Alpha, (const TCHAR*)Fmt, Args...);
+	}
+
+private:
+	void VARARGS PrintfImpl(const FColor& Color, const float Alpha, const TCHAR* Args, ...);
+	void VARARGS PrintfAtImpl(float PosX, float PosY, const FColor& Color, const float Alpha, const TCHAR* Args, ...);
+
+public:
+	// moves cursor to new line
+	void MoveToNewLine();
+
+	// calculate size of string
+	void MeasureString(const FString& String, float& OutSizeX, float& OutSizeY) const;
+
+	// get height of single line text
+	float GetLineHeight() const;
+
+	// project world location on canvas
+	FVector2D ProjectLocation(const FVector& Location) const;
+
+	// check if world location is visible in current view
+	bool IsLocationVisible(const FVector& Location) const;
+
+	// draw item on canvas
+	void DrawItem(FCanvasItem& Item, float PosX, float PosY);
+
+	// draw icon on canvas
+	void DrawIcon(const FColor& Color, const FCanvasIcon& Icon, float PosX, float PosY, float Scale = 1.f);
+
+	// fetches the World associated with this context. Will use the World member variable, if set, and 
+	// PlayerController's world otherwise. Note that it can still be null and needs to be tested.
+	UWorld* GetWorld() const;
+};
+
+namespace FGameplayDebuggerCanvasStrings
+{
+	static FString ColorNameInput = TEXT("white");
+	static FString ColorNameEnabled = TEXT("green");
+	static FString ColorNameDisabled = TEXT("grey");
+	static FString ColorNameEnabledActiveRow = TEXT("green");
+	static FString ColorNameDisabledActiveRow = TEXT("black");
+
+	static FString Separator = TEXT("{white} | ");
+	static FString SeparatorSpace = TEXT("  ");
+}
+
+UENUM()
+enum class EGameplayDebuggerShape : uint8
+{
+	Invalid,
+	Point,
+	Segment,
+	Box,
+	Cone,
+	Cylinder,
+	Circle,
+	Capsule,
+	Polygon,
+	Arrow,
+};
+
+USTRUCT()
+struct GAMEPLAYDEBUGGER_API FGameplayDebuggerShape
+{
+	GENERATED_BODY()
+
+	/** points defining shape */
+	UPROPERTY()
+	TArray<FVector> ShapeData;
+
+	/** description of shape */
+	UPROPERTY()
+	FString Description;
+
+	/** color of shape */
+	UPROPERTY()
+	FColor Color;
+
+	/** type of shape */
+	UPROPERTY()
+	EGameplayDebuggerShape Type;
+
+	FGameplayDebuggerShape() : Color(EForceInit::ForceInit), Type(EGameplayDebuggerShape::Invalid) {}
+
+	bool operator==(const FGameplayDebuggerShape& Other) const
+	{
+		return (Type == Other.Type) && (Color == Other.Color) && (Description == Other.Description) && (ShapeData == Other.ShapeData);
+	}
+
+	static FGameplayDebuggerShape MakePoint(const FVector& Location, const float Radius, const FColor& Color, const FString& Description = FString());
+	static FGameplayDebuggerShape MakeSegment(const FVector& StartLocation, const FVector& EndLocation, const float Thickness, const FColor& Color, const FString& Description = FString());
+	static FGameplayDebuggerShape MakeSegment(const FVector& StartLocation, const FVector& EndLocation, const FColor& Color, const FString& Description = FString());
+	static FGameplayDebuggerShape MakeArrow(const FVector& StartLocation, const FVector& EndLocation, const float HeadSize, const float Thickness, const FColor& Color, const FString& Description = FString());
+	static FGameplayDebuggerShape MakeBox(const FVector& Center, const FVector& Extent, const FColor& Color, const FString& Description = FString());
+	static FGameplayDebuggerShape MakeCone(const FVector& Location, const FVector& Direction, const float Length, const FColor& Color, const FString& Description = FString());
+	static FGameplayDebuggerShape MakeCylinder(const FVector& Center, const float Radius, const float HalfHeight, const FColor& Color, const FString& Description = FString());
+	static FGameplayDebuggerShape MakeCircle(const FVector& Center, const FVector& Up, const float Radius, const FColor& Color, const FString& Description = FString());
+	static FGameplayDebuggerShape MakeCircle(const FVector& Center, const FVector& Up, const float Radius, const float Thickness, const FColor& Color, const FString& Description = FString());
+	static FGameplayDebuggerShape MakeCapsule(const FVector& Center, const float Radius, const float HalfHeight, const FColor& Color, const FString& Description = FString());
+	static FGameplayDebuggerShape MakePolygon(const TArray<FVector>& Verts, const FColor& Color, const FString& Description = FString());
+
+	void Draw(UWorld* World, FGameplayDebuggerCanvasContext& Context);
+};
+
+FArchive& operator<<(FArchive& Ar, FGameplayDebuggerShape& Shape);
+
+enum class EGameplayDebuggerDataPack : uint8
+{
+	Persistent,
+	ResetOnActorChange,
+	ResetOnTick,
+};
+
+USTRUCT()
+struct FGameplayDebuggerDataPackHeader
+{
+	GENERATED_BODY()
+
+	/** version, increased every time new Data is requested */
+	UPROPERTY()
+	int16 DataVersion;
+
+	/** debug actor sync counter */
+	UPROPERTY()
+	int16 SyncCounter;
+
+	/** size of Data array */
+	UPROPERTY()
+	int32 DataSize;
+
+	/** offset to currently replicated portion of data */
+	UPROPERTY()
+	int32 DataOffset;
+
+	/** is data compressed? */
+	UPROPERTY()
+	uint32 bIsCompressed : 1;
+
+	FGameplayDebuggerDataPackHeader() : DataVersion(0), SyncCounter(0), DataSize(0), DataOffset(0), bIsCompressed(false) {}
+
+	FORCEINLINE bool Equals(const FGameplayDebuggerDataPackHeader& Other) const
+	{
+		return (DataVersion == Other.DataVersion) && (DataSize == Other.DataSize) && (DataOffset == Other.DataOffset);
+	}
+
+	FORCEINLINE bool operator==(const FGameplayDebuggerDataPackHeader& Other) const
+	{
+		return Equals(Other);
+	}
+
+	FORCEINLINE bool operator!=(const FGameplayDebuggerDataPackHeader& Other) const
+	{
+		return !Equals(Other);
+	}
+};
+
+struct GAMEPLAYDEBUGGER_API FGameplayDebuggerDataPack
+{
+	using FHeader = FGameplayDebuggerDataPackHeader;
+
+	DECLARE_DELEGATE(FOnReset);
+	DECLARE_DELEGATE_OneParam(FOnSerialize, FArchive&);
+
+	FGameplayDebuggerDataPack() :
+		PackId(0), DataCRC(0), bIsDirty(false), bNeedsConfirmation(false), bReceived(false),
+		Flags(EGameplayDebuggerDataPack::ResetOnTick)
+	{}
+
+	/** data being replicated */
+	TArray<uint8> Data;
+
+	/** minimal header used for DataPack state comparison */
+	FHeader Header;
+
+	/** auto assigned Id of pack */
+	int32 PackId;
+
+	/** CRC used to test changes in Data array */
+	uint32 DataCRC;
+
+	/** force net replication, regardless DataCRC */
+	uint32 bIsDirty : 1;
+
+	/** if set, replication must be confirmed by client before sending next update */
+	uint32 bNeedsConfirmation : 1;
+
+	/** set when client receives complete data pack */
+	uint32 bReceived : 1;
+
+	/** data pack flags */
+	EGameplayDebuggerDataPack Flags;
+
+	FOnReset ResetDelegate;
+	FOnSerialize SerializeDelegate;
+
+	static int32 PacketSize;
+
+	bool CheckDirtyAndUpdate();
+	bool RequestReplication(int16 SyncCounter);
+
+	void OnReplicated();
+	void OnPacketRequest(int16 DataVersion, int32 DataOffset);
+
+	/** get replication progress in (0..1) range */
+	FORCEINLINE float GetProgress() const
+	{
+		return (Header.DataOffset != Header.DataSize) ? ((1.0f * Header.DataOffset) / Header.DataSize) : 1.0f;
+	}
+
+	/** is replication in progress? */
+	FORCEINLINE bool IsInProgress() const
+	{
+		return (Header.DataSize > 0) && (Header.DataOffset < Header.DataSize);
+	}
+
+	static bool IsMultiPacket(int32 TestSize)
+	{
+		return TestSize > PacketSize;
+	}
+};
+
+enum class EGameplayDebuggerInputMode : uint8
+{
+	// input handler is called on local category
+	Local,
+
+	// input handler is replicated to authority category and called there
+	Replicated,
+};
+
+struct GAMEPLAYDEBUGGER_API FGameplayDebuggerInputModifier
+{
+	uint32 bShift : 1;
+	uint32 bCtrl : 1;
+	uint32 bAlt : 1;
+	uint32 bCmd : 1;
+	uint32 bPressed : 1;
+	uint32 bReleased : 1;
+
+	FGameplayDebuggerInputModifier() : bShift(false), bCtrl(false), bAlt(false), bCmd(false), bPressed(true), bReleased(false) {}
+	FGameplayDebuggerInputModifier(bool bInShift, bool bInCtrl, bool bInAlt, bool bInCmd) : bShift(bInShift), bCtrl(bInCtrl), bAlt(bInAlt), bCmd(bInCmd), bPressed(true), bReleased(false) {}
+
+	FGameplayDebuggerInputModifier operator+(const FGameplayDebuggerInputModifier& Other) const
+	{
+		return FGameplayDebuggerInputModifier(bShift || Other.bShift, bCtrl || Other.bCtrl, bAlt || Other.bAlt, bCmd || Other.bCmd);
+	}
+
+	static FGameplayDebuggerInputModifier Shift;
+	static FGameplayDebuggerInputModifier Ctrl;
+	static FGameplayDebuggerInputModifier Alt;
+	static FGameplayDebuggerInputModifier Cmd;
+	static FGameplayDebuggerInputModifier None;
+};
+
+struct GAMEPLAYDEBUGGER_API FGameplayDebuggerInputHandler
+{
+	DECLARE_DELEGATE(FHandler);
+
+	FName KeyName;
+	FGameplayDebuggerInputModifier Modifier;
+	FHandler Delegate;
+	EGameplayDebuggerInputMode Mode;
+
+	FGameplayDebuggerInputHandler() : KeyName(NAME_None), Mode(EGameplayDebuggerInputMode::Local) {}
+
+	bool IsValid() const;
+	FString ToString() const;
+};
+
+/** 
+ * Customizable key binding used by FGameplayDebuggerAddonBase (both categories and extensions)
+ * Intended to use only from addon's constructor!
+ * Check example in FGameplayDebuggerExtension_Spectator::FGameplayDebuggerExtension_Spectator() for details
+ */
+struct GAMEPLAYDEBUGGER_API FGameplayDebuggerInputHandlerConfig
+{
+	FName KeyName;
+	FGameplayDebuggerInputModifier Modifier;
+
+	FGameplayDebuggerInputHandlerConfig() : KeyName(NAME_None) {}
+	FGameplayDebuggerInputHandlerConfig(const FName ConfigName, const FName DefaultKeyName);
+	FGameplayDebuggerInputHandlerConfig(const FName ConfigName, const FName DefaultKeyName, const FGameplayDebuggerInputModifier& DefaultModifier);
+
+	static FName CurrentCategoryName;
+	static FName CurrentExtensionName;
+
+private:
+	void UpdateConfig(const FName ConfigName);
+};
