Disabling Gameplay Debugger input handling, default categories and debug draw when Gameplay Debugger is set to Core mode.

diff --git a/Engine/Source/Editor/GameplayDebugger/Private/GameplayDebuggerEdMode.cpp b/Engine/Source/Editor/GameplayDebugger/Private/GameplayDebuggerEdMode.cpp
--- a/Engine/Source/Editor/GameplayDebugger/Private/GameplayDebuggerEdMode.cpp
+++ b/Engine/Source/Editor/GameplayDebugger/Private/GameplayDebuggerEdMode.cpp
@@ -66,6 +66,7 @@ bool FGameplayDebuggerEdMode::LostFocus(FEditorViewportClient* ViewportClient, F
 
 bool FGameplayDebuggerEdMode::InputKey(FEditorViewportClient* InViewportClient, FViewport* InViewport, FKey InKey, EInputEvent InEvent)
 {
+#if WITH_GAMEPLAY_DEBUGGER
 	UWorld* MyWorld = GetWorld();
 	APlayerController* LocalPC = GEngine->GetFirstLocalPlayerController(MyWorld);
 
@@ -96,7 +97,7 @@ bool FGameplayDebuggerEdMode::InputKey(FEditorViewportClient* InViewportClient,
 			return true;			
 		}
 	}
-
+#endif // WITH_GAMEPLAY_DEBUGGER
 	return false;
 }
 
diff --git a/Engine/Source/Runtime/AIModule/Private/AIModule.cpp b/Engine/Source/Runtime/AIModule/Private/AIModule.cpp
--- a/Engine/Source/Runtime/AIModule/Private/AIModule.cpp
+++ b/Engine/Source/Runtime/AIModule/Private/AIModule.cpp
@@ -13,11 +13,13 @@
 #if ENABLE_VISUAL_LOG
 	#include "VisualLoggerExtension.h"
 #endif // ENABLE_VISUAL_LOG
-#endif
+#endif // WITH_EDITOR
 
+#if WITH_GAMEPLAY_DEBUGGER_CORE
+#include "GameplayDebugger.h"
+#endif // WITH_GAMEPLAY_DEBUGGER_CORE
 
 #if WITH_GAMEPLAY_DEBUGGER
-#include "GameplayDebugger.h"
 #include "GameplayDebugger/GameplayDebuggerCategory_AI.h"
 #include "GameplayDebugger/GameplayDebuggerCategory_BehaviorTree.h"
 #include "GameplayDebugger/GameplayDebuggerCategory_EQS.h"
@@ -80,7 +82,7 @@ void FAIModule::StartupModule()
 	GameplayDebuggerModule.RegisterCategory("Perception", IGameplayDebugger::FOnGetCategory::CreateStatic(&FGameplayDebuggerCategory_Perception::MakeInstance));
 	GameplayDebuggerModule.RegisterCategory("NavGrid", IGameplayDebugger::FOnGetCategory::CreateStatic(&FGameplayDebuggerCategory_NavLocalGrid::MakeInstance), EGameplayDebuggerCategoryState::Hidden);
 	GameplayDebuggerModule.NotifyCategoriesChanged();
-#endif
+#endif // WITH_GAMEPLAY_DEBUGGER
 }
 
 void FAIModule::ShutdownModule()
@@ -103,7 +105,7 @@ void FAIModule::ShutdownModule()
 		GameplayDebuggerModule.UnregisterCategory("NavGrid");
 		GameplayDebuggerModule.NotifyCategoriesChanged();
 	}
-#endif
+#endif // WITH_GAMEPLAY_DEBUGGER
 }
 
 UAISystemBase* FAIModule::CreateAISystemInstance(UWorld* World)
diff --git a/Engine/Source/Runtime/GameplayDebugger/Private/GameplayDebuggerLocalController.cpp b/Engine/Source/Runtime/GameplayDebugger/Private/GameplayDebuggerLocalController.cpp
--- a/Engine/Source/Runtime/GameplayDebugger/Private/GameplayDebuggerLocalController.cpp
+++ b/Engine/Source/Runtime/GameplayDebugger/Private/GameplayDebuggerLocalController.cpp
@@ -44,15 +44,18 @@
 	bPrevLocallyEnabled = false;
 	bEnableTextShadow = false;
 	bPrevScreenMessagesEnabled = false;
-	ActiveRowIdx = 0;
 #if WITH_EDITOR
 	bActivateOnPIEEnd = false;
 #endif // WITH_EDITOR
+
+#if WITH_GAMEPLAY_DEBUGGER
+	ActiveRowIdx = 0;
 	if (HasAnyFlags(RF_ClassDefaultObject) == false)
 	{
 		HUDFont = NewObject<UFont>(this, TEXT("HUDFont"), RF_NoFlags, GEngine->GetSmallFont());
 		HUDFont->LegacyFontSize = UGameplayDebuggerUserSettings::GetFontSize(); //FGameplayDebuggerTweakables::FontSize;
 	}
+#endif // WITH_GAMEPLAY_DEBUGGER
 }
 
 void UGameplayDebuggerLocalController::Initialize(AGameplayDebuggerCategoryReplicator& Replicator, AGameplayDebuggerPlayerManager& Manager)
@@ -61,6 +64,7 @@
 	CachedPlayerManager = &Manager;
 	bSimulateMode = FGameplayDebuggerAddonBase::IsSimulateInEditor() || Replicator.IsEditorWorldReplicator();
 
+#if WITH_GAMEPLAY_DEBUGGER
 	UDebugDrawService::Register(bSimulateMode ? TEXT("DebugAI") : TEXT("Game"), FDebugDrawDelegate::CreateUObject(this, &UGameplayDebuggerLocalController::OnDebugDraw));
 
 #if WITH_EDITOR
@@ -83,10 +87,6 @@
 	}
 #endif
 
-	FGameplayDebuggerAddonManager& AddonManager = FGameplayDebuggerAddonManager::GetCurrent();
-	AddonManager.OnCategoriesChanged.AddUObject(this, &UGameplayDebuggerLocalController::OnCategoriesChanged);
-	OnCategoriesChanged();
-
 	const UGameplayDebuggerConfig* SettingsCDO = UGameplayDebuggerConfig::StaticClass()->GetDefaultObject<UGameplayDebuggerConfig>();
 	const FKey NumpadKeys[] = { EKeys::NumPadZero, EKeys::NumPadOne, EKeys::NumPadTwo, EKeys::NumPadThree, EKeys::NumPadFour,
 		EKeys::NumPadFive, EKeys::NumPadSix, EKeys::NumPadSeven, EKeys::NumPadEight, EKeys::NumPadNine };
@@ -124,6 +124,11 @@
 	PaddingBottom = SettingsCDO->DebugCanvasPaddingBottom;
 
 	bEnableTextShadow = SettingsCDO->bDebugCanvasEnableTextShadow;
+#endif // WITH_GAMEPLAY_DEBUGGER
 
+	FGameplayDebuggerAddonManager& AddonManager = FGameplayDebuggerAddonManager::GetCurrent();
+	AddonManager.OnCategoriesChanged.AddUObject(this, &UGameplayDebuggerLocalController::OnCategoriesChanged);
+	OnCategoriesChanged();
+
 	bNeedsCleanup = true;
 }
@@ -152,9 +157,38 @@
 	}
 }
 
+void UGameplayDebuggerLocalController::OnCategoriesChanged()
+{
+	FGameplayDebuggerAddonManager& AddonManager = FGameplayDebuggerAddonManager::GetCurrent();
+
+	SlotNames.Reset();
+	SlotNames.Append(AddonManager.GetSlotNames());
+
+	// categories are already sorted using AddonManager.SlotMap, build Slot to Category Id map accordingly
+	const TArray< TArray<int32> >& SlotMap = AddonManager.GetSlotMap();
+	SlotCategoryIds.Reset();
+	SlotCategoryIds.AddDefaulted(SlotMap.Num());
+
+	int32 CategoryId = 0;
+	for (int32 SlotIdx = 0; SlotIdx < SlotMap.Num(); SlotIdx++)
+	{
+		for (int32 InnerIdx = 0; InnerIdx < SlotMap[SlotIdx].Num(); InnerIdx++)
+		{
+			SlotCategoryIds[SlotIdx].Add(CategoryId);
+			CategoryId++;
+		}
+	}
+
+	NumCategorySlots = SlotCategoryIds.Num();
+	NumCategories = AddonManager.GetNumVisibleCategories();
+
+	DataPackMap.Reset();
+}
+
+#if WITH_GAMEPLAY_DEBUGGER
 void UGameplayDebuggerLocalController::OnDebugDraw(class UCanvas* Canvas, class APlayerController* PC)
 {
-	if (CachedReplicator && CachedReplicator->IsEnabled())
+	if (CachedReplicator && CachedReplicator->IsEnabled() && bDebugDrawEnabled)
 	{
 		FGameplayDebuggerCanvasContext CanvasContext(Canvas, HUDFont);
 		CanvasContext.CursorX = CanvasContext.DefaultX = PaddingLeft;
@@ -799,35 +833,7 @@
 		CachedReplicator->CollectCategoryData(/*bForce=*/true);
 	}
 }
-#endif
+#endif // WITH_EDITOR
 
-void UGameplayDebuggerLocalController::OnCategoriesChanged()
-{
-	FGameplayDebuggerAddonManager& AddonManager = FGameplayDebuggerAddonManager::GetCurrent();
-
-	SlotNames.Reset();
-	SlotNames.Append(AddonManager.GetSlotNames());
-
-	// categories are already sorted using AddonManager.SlotMap, build Slot to Category Id map accordingly
-	const TArray< TArray<int32> >& SlotMap = AddonManager.GetSlotMap();
-	SlotCategoryIds.Reset();
-	SlotCategoryIds.AddDefaulted(SlotMap.Num());
-
-	int32 CategoryId = 0;
-	for (int32 SlotIdx = 0; SlotIdx < SlotMap.Num(); SlotIdx++)
-	{
-		for (int32 InnerIdx = 0; InnerIdx < SlotMap[SlotIdx].Num(); InnerIdx++)
-		{
-			SlotCategoryIds[SlotIdx].Add(CategoryId);
-			CategoryId++;
-		}
-	}
-
-	NumCategorySlots = SlotCategoryIds.Num();
-	NumCategories = AddonManager.GetNumVisibleCategories();
-
-	DataPackMap.Reset();
-}
-
 void UGameplayDebuggerLocalController::RebuildDataPackMap()
 {
@@ -1098,3 +1104,4 @@
 	FConsoleCommandWithWorldAndArgsDelegate::CreateStatic(&FGameplayDebuggerConsoleCommands::SetFontSize)
 );
 
+#endif // WITH_GAMEPLAY_DEBUGGER
diff --git a/Engine/Source/Runtime/GameplayDebugger/Private/GameplayDebuggerPlayerManager.cpp b/Engine/Source/Runtime/GameplayDebugger/Private/GameplayDebuggerPlayerManager.cpp
--- a/Engine/Source/Runtime/GameplayDebugger/Private/GameplayDebuggerPlayerManager.cpp
+++ b/Engine/Source/Runtime/GameplayDebugger/Private/GameplayDebuggerPlayerManager.cpp
@@ -165,6 +165,7 @@ void AGameplayDebuggerPlayerManager::UpdateAuthReplicators()
 		}
 	}
 
+#if WITH_GAMEPLAY_DEBUGGER
 	for (FConstPlayerControllerIterator It = World->GetPlayerControllerIterator(); It; It++)
 	{
 		APlayerController* TestPC = It->Get();
@@ -179,6 +180,7 @@ void AGameplayDebuggerPlayerManager::UpdateAuthReplicators()
 			}
 		}
 	}
+#endif // WITH_GAMEPLAY_DEBUGGER
 
 	PrimaryActorTick.TickInterval = PlayerData.Num() ? 5.0f : 0.5f;
 }
@@ -194,7 +196,8 @@ void AGameplayDebuggerPlayerManager::RegisterReplicator(AGameplayDebuggerCategor
 	// keep all player related objects together for easy access and GC
 	FGameplayDebuggerPlayerData NewData;
 	NewData.Replicator = &Replicator;
-	
+
+#if WITH_GAMEPLAY_DEBUGGER
 	if (bIsLocal)
 	{
 		APlayerController* OwnerPC = Replicator.GetReplicationOwner();
@@ -216,12 +219,14 @@ void AGameplayDebuggerPlayerManager::RegisterReplicator(AGameplayDebuggerCategor
 		NewData.Controller = nullptr;
 		NewData.InputComponent = nullptr;
 	}
+#endif // WITH_GAMEPLAY_DEBUGGER
 
 	PlayerData.Add(NewData);
 }
 
 void AGameplayDebuggerPlayerManager::RefreshInputBindings(AGameplayDebuggerCategoryReplicator& Replicator)
 {
+#if WITH_GAMEPLAY_DEBUGGER
 	for (int32 Idx = 0; Idx < PlayerData.Num(); Idx++)
 	{
 		FGameplayDebuggerPlayerData& TestData = PlayerData[Idx];
@@ -234,6 +239,7 @@ void AGameplayDebuggerPlayerManager::RefreshInputBindings(AGameplayDebuggerCateg
 			TestData.Controller->BindInput(*TestData.InputComponent);
 		}
 	}
+#endif // WITH_GAMEPLAY_DEBUGGER
 }
 
 AGameplayDebuggerCategoryReplicator* AGameplayDebuggerPlayerManager::GetReplicator(const APlayerController& OwnerPC) const
diff --git a/Engine/Source/Runtime/GameplayDebugger/Public/GameplayDebugger.h b/Engine/Source/Runtime/GameplayDebugger/Public/GameplayDebugger.h
--- a/Engine/Source/Runtime/GameplayDebugger/Public/GameplayDebugger.h
+++ b/Engine/Source/Runtime/GameplayDebugger/Public/GameplayDebugger.h
@@ -10,7 +10,8 @@
 // Category class:
 // - derives from FGameplayDebuggerCategory
 // - implements at least CollectData() and DrawData() functions
-// - requires WITH_GAMEPLAY_DEBUGGER define to compile (doesn't exist in shipping builds by default; see note below to override)
+// - If WITH_GAMEPLAY_DEBUGGER is defined (doesn't exist in shipping builds by default; see note below to override), it will include all default debug categories and menu
+// - If WITH_GAMEPLAY_DEBUGGER_CORE is defined (see note below), it will compile the core parts of the tool, waiting for user to register a debug category and create a replicator when desired.
 // - needs to be registered and unregistered manually by owning module
 // - automatically replicate data added with FGameplayDebuggerCategory::AddTextLine, FGameplayDebuggerCategory::AddShape
 // - automatically replicate data structs initialized with FGameplayDebuggerCategory::SetDataPackReplication
@@ -27,7 +28,8 @@
 // Check AIModule/Private/AIModule.cpp for registration example.
 //
 // Note. Use 'SetupGameplayDebuggerSupport(Target)' when adding module to your project's Build.cs (see AIModule/AIModule.Build.cs)
-// Note. Use 'bUseGameplayDebugger={0|1}' in your <ProjectTargetType>.Target.cs to force GameplayDebugger disabled/enabled
+// Note. Use 'bUseGameplayDebugger={0|1}' in your <ProjectTargetType>.Target.cs to force GameplayDebugger disabled/enabled (if enabled, it will implicitly set bUseGameplayDebuggerCore=1)
+// Note. Use 'bUseGameplayDebuggerCore={0|1}' in your <ProjectTargetType>.Target.cs to compile on/off the GameplayDebugger core parts only
 // 
 
 #pragma once
diff --git a/Engine/Source/Runtime/GameplayDebugger/Public/GameplayDebuggerAddonManager.h b/Engine/Source/Runtime/GameplayDebugger/Public/GameplayDebuggerAddonManager.h
--- a/Engine/Source/Runtime/GameplayDebugger/Public/GameplayDebuggerAddonManager.h
+++ b/Engine/Source/Runtime/GameplayDebugger/Public/GameplayDebuggerAddonManager.h
@@ -71,7 +71,7 @@ public:
 	int32 GetNumVisibleCategories() const { return NumVisibleCategories; }
 
 	/** singleton accessor */
-	static FGameplayDebuggerAddonManager& GetCurrent();
+	static GAMEPLAYDEBUGGER_API FGameplayDebuggerAddonManager& GetCurrent();
 
 	/** event called when CategoryMap changes */
 	FOnGameplayDebuggerAddonEvent OnCategoriesChanged;
diff --git a/Engine/Source/Runtime/GameplayDebugger/Public/GameplayDebuggerLocalController.h b/Engine/Source/Runtime/GameplayDebugger/Public/GameplayDebuggerLocalController.h
--- a/Engine/Source/Runtime/GameplayDebugger/Public/GameplayDebuggerLocalController.h
+++ b/Engine/Source/Runtime/GameplayDebugger/Public/GameplayDebuggerLocalController.h
@@ -17,7 +17,6 @@ class UInputComponent;
 struct FKey;
 class UFont;
 
-
 UCLASS(NotBlueprintable, NotBlueprintType, noteditinlinenew, hidedropdown, Transient)
 class UGameplayDebuggerLocalController : public UObject
 {
@@ -31,6 +30,10 @@ class UGameplayDebuggerLocalController : public UObject
 	/** remove from world */
 	void Cleanup();
 
+#if WITH_GAMEPLAY_DEBUGGER
+	/** Whether the Debug Draw rendering is allowed for this replicator */
+	void SetDebugDrawEnabled(const bool bInDebugDrawEnabled) { bDebugDrawEnabled = bInDebugDrawEnabled; }
+
 	/** drawing event */
 	void OnDebugDraw(class UCanvas* Canvas, class APlayerController* PC);
 
@@ -39,10 +42,9 @@ class UGameplayDebuggerLocalController : public UObject
 
 	/** checks if key is bound by any action */
 	bool GAMEPLAYDEBUGGER_API IsKeyBound(const FName KeyName) const;
+#endif // WITH_GAMEPLAY_DEBUGGER
 
 protected:
-	friend struct FGameplayDebuggerConsoleCommands;
-
 	UPROPERTY()
 	TObjectPtr<AGameplayDebuggerCategoryReplicator> CachedReplicator;
 
@@ -59,8 +61,6 @@ protected:
 	TArray<TArray<int32> > SlotCategoryIds;
 	TArray<FString> SlotNames;
 
-	TSet<FName> UsedBindings;
-
 	uint32 bSimulateMode : 1;
 	uint32 bNeedsCleanup : 1;
 	uint32 bIsSelectingActor : 1;
@@ -71,6 +71,13 @@ protected:
 #if WITH_EDITOR
 	uint32 bActivateOnPIEEnd : 1;
 #endif // WITH_EDITOR
+	int32 NumCategorySlots;
+	int32 NumCategories;
+
+#if WITH_GAMEPLAY_DEBUGGER
+	friend struct FGameplayDebuggerConsoleCommands;
+
+	TSet<FName> UsedBindings;
 
 	FString ActivationKeyDesc;
 	FString RowUpKeyDesc;
@@ -78,8 +85,6 @@ protected:
 	FString CategoryKeysDesc;
 
 	int32 ActiveRowIdx;
-	int32 NumCategorySlots;
-	int32 NumCategories;
 	static constexpr int32 NumCategoriesPerRow = 10;
 
 	float PaddingLeft;
@@ -90,6 +95,8 @@ protected:
 	FTimerHandle StartSelectingActorHandle;
 	FTimerHandle SelectActorTickHandle;
 
+	bool bDebugDrawEnabled = true;
+
 	void OnActivationPressed();
 	void OnActivationReleased();
 	void OnCategory0Pressed();
@@ -134,14 +141,15 @@ protected:
 	void OnSelectedObject(UObject* Object);
 	void OnBeginPIE(const bool bIsSimulating);
 	void OnEndPIE(const bool bIsSimulating);
-#endif
+#endif // WITH_EDITOR
 
 	FString GetKeyDescriptionShort(const FKey& KeyBind) const;
 	FString GetKeyDescriptionLong(const FKey& KeyBind) const;
+	
+	/** build DataPackMap for replication details */
+	void RebuildDataPackMap();
+#endif // WITH_GAMEPLAY_DEBUGGER
 
 	/** called when known category set has changed */
 	void OnCategoriesChanged();
-
-	/** build DataPackMap for replication details */
-	void RebuildDataPackMap();
 };
\ No newline at end of file
