[GameplayDebugger] marked back 'Actor' as UPROPERTY in  FGameplayDebuggerDebugActor since it is required to be replicated. The initial fix was to remove the hard reference to the actor and the weak object ptr still fixes that part.
#rnx

diff --git a/Engine/Source/Runtime/GameplayDebugger/Public/GameplayDebuggerCategoryReplicator.h b/Engine/Source/Runtime/GameplayDebugger/Public/GameplayDebuggerCategoryReplicator.h
--- a/Engine/Source/Runtime/GameplayDebugger/Public/GameplayDebuggerCategoryReplicator.h
+++ b/Engine/Source/Runtime/GameplayDebugger/Public/GameplayDebuggerCategoryReplicator.h
@@ -95,7 +95,7 @@ struct FGameplayDebuggerDebugActor
 {
 	GENERATED_USTRUCT_BODY()
 
-	UPROPERTY(NotReplicated)
+	UPROPERTY()
 	TWeakObjectPtr<AActor> Actor;
 
 	UPROPERTY()
