PR #9701: Fix DPI scaling issues in FGameplayDebuggerCanvasContext::MeasureString

diff --git a/Engine/Source/Runtime/GameplayDebugger/Private/GameplayDebuggerTypes.cpp b/Engine/Source/Runtime/GameplayDebugger/Private/GameplayDebuggerTypes.cpp
--- a/Engine/Source/Runtime/GameplayDebugger/Private/GameplayDebuggerTypes.cpp
+++ b/Engine/Source/Runtime/GameplayDebugger/Private/GameplayDebuggerTypes.cpp
@@ -564,6 +564,7 @@ void FGameplayDebuggerCanvasContext::MeasureString(const FString& String, float&
 			}
 		}
 
+		const float DPIScale = CanvasOb->GetDPIScale();
 		TArray<FString> Lines;
 		StringWithoutFormatting.ParseIntoArrayLines(Lines);
 		
@@ -571,7 +572,9 @@ void FGameplayDebuggerCanvasContext::MeasureString(const FString& String, float&
 		for (int32 Idx = 0; Idx < Lines.Num(); Idx++)
 		{
 			float LineSizeX = 0.0f, LineSizeY = 0.0f;
-			CanvasOb->StrLen(FontOb, Lines[Idx], LineSizeX, LineSizeY);
+			CanvasOb->StrLen(FontOb, Lines[Idx], LineSizeX, LineSizeY, /* bDPIAwareStringMeasurement */ true);
+			LineSizeX /= DPIScale;
+			LineSizeY /= DPIScale;
 
 			OutSizeX = FMath::Max(OutSizeX, LineSizeX);
 			OutSizeY += LineSizeY;
