Fixed unnecessary copies returned from FGameplayDebuggerAddonManager::GetSlotNames().

diff --git a/Engine/Source/Runtime/GameplayDebugger/Public/GameplayDebuggerAddonManager.h b/Engine/Source/Runtime/GameplayDebugger/Public/GameplayDebuggerAddonManager.h
--- a/Engine/Source/Runtime/GameplayDebugger/Public/GameplayDebuggerAddonManager.h
+++ b/Engine/Source/Runtime/GameplayDebugger/Public/GameplayDebuggerAddonManager.h
@@ -65,7 +65,7 @@ public:
 	const TArray<TArray<int32> >& GetSlotMap() const { return SlotMap; }
 
 	/** get slot-Name map */
-	const TArray<FString> GetSlotNames() const { return SlotNames; }
+	const TArray<FString>& GetSlotNames() const { return SlotNames; }
 
 	/** get number of visible categories */
 	int32 GetNumVisibleCategories() const { return NumVisibleCategories; }
