PR #9852: Fix GameplayDebugger header so it incorporates DPI correctly for displaying time and VLog info at the intended locations

diff --git a/Engine/Source/Runtime/GameplayDebugger/Private/GameplayDebuggerLocalController.cpp b/Engine/Source/Runtime/GameplayDebugger/Private/GameplayDebuggerLocalController.cpp
--- a/Engine/Source/Runtime/GameplayDebugger/Private/GameplayDebuggerLocalController.cpp
+++ b/Engine/Source/Runtime/GameplayDebugger/Private/GameplayDebuggerLocalController.cpp
@@ -240,7 +240,8 @@ void UGameplayDebuggerLocalController::DrawHeader(FGameplayDebuggerCanvasContext
 	const float LineHeight = CanvasContext.GetLineHeight();
 	const int32 NumExtensions = bSimulateMode ? 0 : CachedReplicator->GetNumExtensions();
 	const int32 NumExtensionRows = (NumExtensions > 0) ? 1 : 0;
-	const float CanvasSizeX = CanvasContext.Canvas->SizeX - PaddingLeft - PaddingRight;
+	const float DPIScale = CanvasContext.Canvas->GetDPIScale();
+	const float CanvasSizeX = (CanvasContext.Canvas->SizeX / DPIScale) - PaddingLeft - PaddingRight;
 	const float UsePaddingTop = PaddingTop + (bSimulateMode ? 30.0f : 0);
 	
 	const float BackgroundPadding = 5.0f;
@@ -285,14 +286,14 @@ void UGameplayDebuggerLocalController::DrawHeader(FGameplayDebuggerCanvasContext
 	const FString DebugActorDesc = FString::Printf(TEXT("Debug actor: {cyan}%s"), *CachedReplicator->GetDebugActorName().ToString());
 	float DebugActorSizeX = 0.0f, DebugActorSizeY = 0.0f;
 	CanvasContext.MeasureString(DebugActorDesc, DebugActorSizeX, DebugActorSizeY);
-	CanvasContext.PrintAt(CanvasContext.Canvas->SizeX - PaddingRight - DebugActorSizeX, UsePaddingTop, DebugActorDesc);
+	CanvasContext.PrintAt((CanvasContext.Canvas->SizeX / DPIScale) - PaddingRight - DebugActorSizeX, UsePaddingTop, DebugActorDesc);
 
 	const FString VLogDesc = FString::Printf(TEXT("VLog: {cyan}%s"), CachedReplicator->GetVisLogSyncData().DeviceIDs.Len() > 0
 			? *CachedReplicator->GetVisLogSyncData().DeviceIDs
 			: TEXT("not recording to file"));
 		float VLogSizeX = 0.0f, VLogSizeY = 0.0f;
 		CanvasContext.MeasureString(VLogDesc, VLogSizeX, VLogSizeY);
-		CanvasContext.PrintAt(CanvasContext.Canvas->SizeX - PaddingRight - VLogSizeX, UsePaddingTop + LineHeight, VLogDesc);
+		CanvasContext.PrintAt((CanvasContext.Canvas->SizeX / DPIScale) - PaddingRight - VLogSizeX, UsePaddingTop + LineHeight, VLogDesc);
 
 	const FString TimestampDesc = FString::Printf(TEXT("Time: %.2fs"), CachedReplicator->GetWorld()->GetTimeSeconds());
 	float TimestampSizeX = 0.0f, TimestampSizeY = 0.0f;
@@ -304,7 +305,7 @@ void UGameplayDebuggerLocalController::DrawHeader(FGameplayDebuggerCanvasContext
 		const FString ChangeRowDesc = FString::Printf(TEXT("Prev row: {yellow}%s\n{white}Next row: {yellow}%s"), *RowUpKeyDesc, *RowDownKeyDesc);
 		float RowDescSizeX = 0.0f, RowDescSizeY = 0.0f;
 		CanvasContext.MeasureString(ChangeRowDesc, RowDescSizeX, RowDescSizeY);
-		CanvasContext.PrintAt(CanvasContext.Canvas->SizeX - PaddingRight - RowDescSizeX, UsePaddingTop + LineHeight * (NumExtensionRows + 1), ChangeRowDesc);
+		CanvasContext.PrintAt((CanvasContext.Canvas->SizeX / DPIScale) - PaddingRight - RowDescSizeX, UsePaddingTop + LineHeight * (NumExtensionRows + 1), ChangeRowDesc);
 	}
 
 	if (NumExtensionRows)
