Allow Gameplay Debugger categories to be compiled in Gameplay Debugger Core mode but never in Shipping.
This will allow game teams to create on-the-fly the Gameplay Debugger when desired (like in Test build) and add the Engine categories they need for the core mode.
Also moved the  Gameplay Debugger Categories header file their corresponding public folder.

diff --git a/Engine/Plugins/AI/HTNPlanner/Source/HTNPlanner/Private/Debug/GameplayDebuggerCategory_HTN.cpp b/Engine/Plugins/AI/HTNPlanner/Source/HTNPlanner/Private/Debug/GameplayDebuggerCategory_HTN.cpp
--- a/Engine/Plugins/AI/HTNPlanner/Source/HTNPlanner/Private/Debug/GameplayDebuggerCategory_HTN.cpp
+++ b/Engine/Plugins/AI/HTNPlanner/Source/HTNPlanner/Private/Debug/GameplayDebuggerCategory_HTN.cpp
@@ -2,7 +2,7 @@
 
 #include "Debug/GameplayDebuggerCategory_HTN.h"
 
-#if WITH_GAMEPLAY_DEBUGGER
+#if WITH_GAMEPLAY_DEBUGGER_MENU
 
 FGameplayDebuggerCategory_HTN::FGameplayDebuggerCategory_HTN()
 {
@@ -22,4 +22,4 @@ void FGameplayDebuggerCategory_HTN::DrawData(APlayerController* OwnerPC, FGamepl
 
 }
 
-#endif // WITH_GAMEPLAY_DEBUGGER
+#endif // WITH_GAMEPLAY_DEBUGGER_MENU
diff --git a/Engine/Plugins/AI/HTNPlanner/Source/HTNPlanner/Public/Debug/GameplayDebuggerCategory_HTN.h b/Engine/Plugins/AI/HTNPlanner/Source/HTNPlanner/Public/Debug/GameplayDebuggerCategory_HTN.h
--- a/Engine/Plugins/AI/HTNPlanner/Source/HTNPlanner/Public/Debug/GameplayDebuggerCategory_HTN.h
+++ b/Engine/Plugins/AI/HTNPlanner/Source/HTNPlanner/Public/Debug/GameplayDebuggerCategory_HTN.h
@@ -2,8 +2,9 @@
 
 #pragma once
 
-#if WITH_GAMEPLAY_DEBUGGER
+#include "CoreMinimal.h"
 
+#if WITH_GAMEPLAY_DEBUGGER_MENU
 #include "GameplayDebuggerCategory.h"
 
 class ABotObjectiveGraph;
@@ -21,4 +22,4 @@ public:
 protected:
 };
 
-#endif // WITH_GAMEPLAY_DEBUGGER
+#endif // WITH_GAMEPLAY_DEBUGGER_MENU
diff --git a/Engine/Plugins/AI/MLAdapter/Source/MLAdapter/Private/Debug/GameplayDebuggerCategory_MLAdapter.cpp b/Engine/Plugins/AI/MLAdapter/Source/MLAdapter/Private/Debug/GameplayDebuggerCategory_MLAdapter.cpp
--- a/Engine/Plugins/AI/MLAdapter/Source/MLAdapter/Private/Debug/GameplayDebuggerCategory_MLAdapter.cpp
+++ b/Engine/Plugins/AI/MLAdapter/Source/MLAdapter/Private/Debug/GameplayDebuggerCategory_MLAdapter.cpp
@@ -9,7 +9,7 @@
 #include "Sensors/MLAdapterSensor.h"
 #include "Actuators/MLAdapterActuator.h"
 
-#if WITH_GAMEPLAY_DEBUGGER
+#if WITH_GAMEPLAY_DEBUGGER_MENU
 #include "GameplayDebuggerCategoryReplicator.h"
 
 //----------------------------------------------------------------------//
@@ -186,5 +186,5 @@ void FGameplayDebuggerCategory_MLAdapter::OnBeginAgentRemove(UMLAdapterAgent& Ag
 	}
 }
 
-#endif // WITH_GAMEPLAY_DEBUGGER
+#endif // WITH_GAMEPLAY_DEBUGGER_MENU
 
diff --git a/Engine/Plugins/AI/MLAdapter/Source/MLAdapter/Public/Debug/GameplayDebuggerCategory_MLAdapter.h b/Engine/Plugins/AI/MLAdapter/Source/MLAdapter/Public/Debug/GameplayDebuggerCategory_MLAdapter.h
--- a/Engine/Plugins/AI/MLAdapter/Source/MLAdapter/Public/Debug/GameplayDebuggerCategory_MLAdapter.h
+++ b/Engine/Plugins/AI/MLAdapter/Source/MLAdapter/Public/Debug/GameplayDebuggerCategory_MLAdapter.h
@@ -2,8 +2,9 @@
 
 #pragma once
 
-#if WITH_GAMEPLAY_DEBUGGER
+#include "CoreMinimal.h"
 
+#if WITH_GAMEPLAY_DEBUGGER_MENU
 #include "GameplayDebuggerCategory.h"
 #include "MLAdapterTypes.h"
 
@@ -38,4 +39,4 @@ protected:
 	FMLAdapter::FAgentID CachedAgentID;
 };
 
-#endif // WITH_GAMEPLAY_DEBUGGER
+#endif // WITH_GAMEPLAY_DEBUGGER_MENU
diff --git a/Engine/Plugins/Runtime/GameplayAbilities/Source/GameplayAbilities/Private/GameplayDebuggerCategory_Abilities.cpp b/Engine/Plugins/Runtime/GameplayAbilities/Source/GameplayAbilities/Private/GameplayDebuggerCategory_Abilities.cpp
--- a/Engine/Plugins/Runtime/GameplayAbilities/Source/GameplayAbilities/Private/GameplayDebuggerCategory_Abilities.cpp
+++ b/Engine/Plugins/Runtime/GameplayAbilities/Source/GameplayAbilities/Private/GameplayDebuggerCategory_Abilities.cpp
@@ -4,7 +4,7 @@
 #include "Engine/Engine.h"
 #include "Engine/GameViewportClient.h"
 
-#if WITH_GAMEPLAY_DEBUGGER
+#if WITH_GAMEPLAY_DEBUGGER_MENU
 
 #include "GameplayTagContainer.h"
 #include "GameplayAbilitySpec.h"
@@ -243,4 +243,4 @@ void FGameplayDebuggerCategory_Abilities::DrawData(APlayerController* OwnerPC, F
 	}
 }
 
-#endif // WITH_GAMEPLAY_DEBUGGER
+#endif // WITH_GAMEPLAY_DEBUGGER_MENU
diff --git a/Engine/Plugins/Runtime/GameplayAbilities/Source/GameplayAbilities/Private/GameplayDebuggerCategory_Abilities.h b/Engine/Plugins/Runtime/GameplayAbilities/Source/GameplayAbilities/Private/GameplayDebuggerCategory_Abilities.h
--- a/Engine/Plugins/Runtime/GameplayAbilities/Source/GameplayAbilities/Private/GameplayDebuggerCategory_Abilities.h
+++ b/Engine/Plugins/Runtime/GameplayAbilities/Source/GameplayAbilities/Private/GameplayDebuggerCategory_Abilities.h
@@ -2,9 +2,9 @@
 
 #pragma once
 
-#if WITH_GAMEPLAY_DEBUGGER
-
 #include "CoreMinimal.h"
+
+#if WITH_GAMEPLAY_DEBUGGER_MENU
 #include "GameplayDebuggerCategory.h"
 
 class AActor;
@@ -52,4 +52,4 @@ protected:
 	bool WrapStringAccordingToViewport(const FString& iStr, FString& oStr, FGameplayDebuggerCanvasContext& CanvasContext, float ViewportWitdh);
 };
 
-#endif // WITH_GAMEPLAY_DEBUGGER
+#endif // WITH_GAMEPLAY_DEBUGGER_MENU
diff --git a/Engine/Source/Programs/UnrealBuildTool/Configuration/ModuleRules.cs b/Engine/Source/Programs/UnrealBuildTool/Configuration/ModuleRules.cs
--- a/Engine/Source/Programs/UnrealBuildTool/Configuration/ModuleRules.cs
+++ b/Engine/Source/Programs/UnrealBuildTool/Configuration/ModuleRules.cs
@@ -1323,6 +1323,14 @@ namespace UnrealBuildTool
 			{
 				PublicDefinitions.Add("WITH_GAMEPLAY_DEBUGGER_CORE=1");
 				PublicDefinitions.Add("WITH_GAMEPLAY_DEBUGGER=" + (Target.bUseGameplayDebugger ? 1 : 0));
+				if (Target.Configuration != UnrealTargetConfiguration.Shipping)
+				{
+					PublicDefinitions.Add("WITH_GAMEPLAY_DEBUGGER_MENU=1");
+				}
+				else
+				{
+					PublicDefinitions.Add("WITH_GAMEPLAY_DEBUGGER_MENU=0");
+				}
 
 				if (bAddAsPublicDependency)
 				{
@@ -1345,6 +1353,7 @@ namespace UnrealBuildTool
 			{
 				PublicDefinitions.Add("WITH_GAMEPLAY_DEBUGGER_CORE=0");
 				PublicDefinitions.Add("WITH_GAMEPLAY_DEBUGGER=0");
+				PublicDefinitions.Add("WITH_GAMEPLAY_DEBUGGER_MENU=0");
 			}
 		}
 
diff --git a/Engine/Source/Runtime/AIModule/Classes/Perception/AIPerceptionComponent.h b/Engine/Source/Runtime/AIModule/Classes/Perception/AIPerceptionComponent.h
--- a/Engine/Source/Runtime/AIModule/Classes/Perception/AIPerceptionComponent.h
+++ b/Engine/Source/Runtime/AIModule/Classes/Perception/AIPerceptionComponent.h
@@ -336,9 +336,9 @@ public:
 	bool HasAnyCurrentStimulus(const AActor& Source) const;
 	bool HasActiveStimulus(const AActor& Source, FAISenseID Sense) const;
 
-#if WITH_GAMEPLAY_DEBUGGER
+#if WITH_GAMEPLAY_DEBUGGER_MENU
 	virtual void DescribeSelfToGameplayDebugger(FGameplayDebuggerCategory* DebuggerCategory) const;
-#endif // WITH_GAMEPLAY_DEBUGGER
+#endif // WITH_GAMEPLAY_DEBUGGER_MENU
 
 #if ENABLE_VISUAL_LOG
 	virtual void DescribeSelfToVisLog(FVisualLogEntry* Snapshot) const;
diff --git a/Engine/Source/Runtime/AIModule/Classes/Perception/AISenseConfig.h b/Engine/Source/Runtime/AIModule/Classes/Perception/AISenseConfig.h
--- a/Engine/Source/Runtime/AIModule/Classes/Perception/AISenseConfig.h
+++ b/Engine/Source/Runtime/AIModule/Classes/Perception/AISenseConfig.h
@@ -46,7 +46,7 @@ public:
 	float GetMaxAge() const { return MaxAge == 0.f ? FAIStimulus::NeverHappenedAge : MaxAge; }	
 	bool IsEnabled() const { return bStartsEnabled; }
 
-#if WITH_GAMEPLAY_DEBUGGER
+#if WITH_GAMEPLAY_DEBUGGER_MENU
 	virtual void DescribeSelfToGameplayDebugger(const UAIPerceptionComponent* PerceptionComponent, FGameplayDebuggerCategory* DebuggerCategory) const;
-#endif // WITH_GAMEPLAY_DEBUGGER
+#endif // WITH_GAMEPLAY_DEBUGGER_MENU
 };
diff --git a/Engine/Source/Runtime/AIModule/Classes/Perception/AISenseConfig_Hearing.h b/Engine/Source/Runtime/AIModule/Classes/Perception/AISenseConfig_Hearing.h
--- a/Engine/Source/Runtime/AIModule/Classes/Perception/AISenseConfig_Hearing.h
+++ b/Engine/Source/Runtime/AIModule/Classes/Perception/AISenseConfig_Hearing.h
@@ -39,7 +39,7 @@ public:
 
 	virtual TSubclassOf<UAISense> GetSenseImplementation() const override;
 
-#if WITH_GAMEPLAY_DEBUGGER
+#if WITH_GAMEPLAY_DEBUGGER_MENU
 	virtual void DescribeSelfToGameplayDebugger(const UAIPerceptionComponent* PerceptionComponent, FGameplayDebuggerCategory* DebuggerCategory) const;
-#endif // WITH_GAMEPLAY_DEBUGGER
+#endif // WITH_GAMEPLAY_DEBUGGER_MENU
 };
diff --git a/Engine/Source/Runtime/AIModule/Classes/Perception/AISenseConfig_Sight.h b/Engine/Source/Runtime/AIModule/Classes/Perception/AISenseConfig_Sight.h
--- a/Engine/Source/Runtime/AIModule/Classes/Perception/AISenseConfig_Sight.h
+++ b/Engine/Source/Runtime/AIModule/Classes/Perception/AISenseConfig_Sight.h
@@ -58,7 +58,7 @@ public:
 	virtual void PostEditChangeChainProperty(FPropertyChangedChainEvent& PropertyChangedEvent) override;
 #endif // WITH_EDITOR
 
-#if WITH_GAMEPLAY_DEBUGGER
+#if WITH_GAMEPLAY_DEBUGGER_MENU
 	virtual void DescribeSelfToGameplayDebugger(const UAIPerceptionComponent* PerceptionComponent, FGameplayDebuggerCategory* DebuggerCategory) const;
-#endif // WITH_GAMEPLAY_DEBUGGER
+#endif // WITH_GAMEPLAY_DEBUGGER_MENU
 };
diff --git a/Engine/Source/Runtime/AIModule/Private/GameplayDebugger/GameplayDebuggerCategory_AI.cpp b/Engine/Source/Runtime/AIModule/Private/GameplayDebugger/GameplayDebuggerCategory_AI.cpp
--- a/Engine/Source/Runtime/AIModule/Private/GameplayDebugger/GameplayDebuggerCategory_AI.cpp
+++ b/Engine/Source/Runtime/AIModule/Private/GameplayDebugger/GameplayDebuggerCategory_AI.cpp
@@ -2,7 +2,7 @@
 
 #include "GameplayDebugger/GameplayDebuggerCategory_AI.h"
 
-#if WITH_GAMEPLAY_DEBUGGER
+#if WITH_GAMEPLAY_DEBUGGER_MENU
 
 #include "GameFramework/Pawn.h"
 #include "ShowFlags.h"
@@ -519,4 +519,4 @@ void FGameplayDebuggerCategory_AI::DrawPawnIcons(UWorld* World, AActor* DebugAct
 	}
 }
 
-#endif // WITH_GAMEPLAY_DEBUGGER
+#endif // WITH_GAMEPLAY_DEBUGGER_MENU
diff --git a/Engine/Source/Runtime/AIModule/Private/GameplayDebugger/GameplayDebuggerCategory_AI.h b/Engine/Source/Runtime/AIModule/Private/GameplayDebugger/GameplayDebuggerCategory_AI.h
deleted file mode 100644
diff --git a/Engine/Source/Runtime/AIModule/Private/GameplayDebugger/GameplayDebuggerCategory_BehaviorTree.cpp b/Engine/Source/Runtime/AIModule/Private/GameplayDebugger/GameplayDebuggerCategory_BehaviorTree.cpp
--- a/Engine/Source/Runtime/AIModule/Private/GameplayDebugger/GameplayDebuggerCategory_BehaviorTree.cpp
+++ b/Engine/Source/Runtime/AIModule/Private/GameplayDebugger/GameplayDebuggerCategory_BehaviorTree.cpp
@@ -2,7 +2,7 @@
 
 #include "GameplayDebugger/GameplayDebuggerCategory_BehaviorTree.h"
 
-#if WITH_GAMEPLAY_DEBUGGER
+#if WITH_GAMEPLAY_DEBUGGER_MENU
 
 #include "GameFramework/Pawn.h"
 #include "BrainComponent.h"
@@ -80,4 +80,4 @@ void FGameplayDebuggerCategory_BehaviorTree::DrawData(APlayerController* OwnerPC
 	}
 }
 
-#endif // WITH_GAMEPLAY_DEBUGGER
+#endif // WITH_GAMEPLAY_DEBUGGER_MENU
diff --git a/Engine/Source/Runtime/AIModule/Private/GameplayDebugger/GameplayDebuggerCategory_BehaviorTree.h b/Engine/Source/Runtime/AIModule/Private/GameplayDebugger/GameplayDebuggerCategory_BehaviorTree.h
deleted file mode 100644
diff --git a/Engine/Source/Runtime/AIModule/Private/GameplayDebugger/GameplayDebuggerCategory_EQS.cpp b/Engine/Source/Runtime/AIModule/Private/GameplayDebugger/GameplayDebuggerCategory_EQS.cpp
--- a/Engine/Source/Runtime/AIModule/Private/GameplayDebugger/GameplayDebuggerCategory_EQS.cpp
+++ b/Engine/Source/Runtime/AIModule/Private/GameplayDebugger/GameplayDebuggerCategory_EQS.cpp
@@ -3,7 +3,7 @@
 #include "GameplayDebugger/GameplayDebuggerCategory_EQS.h"
 #include "GameFramework/Pawn.h"
 
-#if WITH_GAMEPLAY_DEBUGGER
+#if WITH_GAMEPLAY_DEBUGGER_MENU
 
 #include "GameFramework/PlayerController.h"
 #include "CanvasItem.h"
@@ -473,4 +473,4 @@ void FGameplayDebuggerCategory_EQS::ToggleDetailView()
 	bShowDetails = !bShowDetails;
 }
 
-#endif // ENABLE_GAMEPLAY_DEBUGGER
+#endif // WITH_GAMEPLAY_DEBUGGER_MENU
diff --git a/Engine/Source/Runtime/AIModule/Private/GameplayDebugger/GameplayDebuggerCategory_EQS.h b/Engine/Source/Runtime/AIModule/Private/GameplayDebugger/GameplayDebuggerCategory_EQS.h
deleted file mode 100644
diff --git a/Engine/Source/Runtime/AIModule/Private/GameplayDebugger/GameplayDebuggerCategory_NavLocalGrid.cpp b/Engine/Source/Runtime/AIModule/Private/GameplayDebugger/GameplayDebuggerCategory_NavLocalGrid.cpp
--- a/Engine/Source/Runtime/AIModule/Private/GameplayDebugger/GameplayDebuggerCategory_NavLocalGrid.cpp
+++ b/Engine/Source/Runtime/AIModule/Private/GameplayDebugger/GameplayDebuggerCategory_NavLocalGrid.cpp
@@ -1,8 +1,8 @@
 // Copyright Epic Games, Inc. All Rights Reserved.
 
-#include "GameplayDebuggerCategory_NavLocalGrid.h"
+#include "GameplayDebugger/GameplayDebuggerCategory_NavLocalGrid.h"
 
-#if WITH_GAMEPLAY_DEBUGGER
+#if WITH_GAMEPLAY_DEBUGGER_MENU
 
 #include "GameFramework/PlayerController.h"
 #include "AIController.h"
@@ -435,4 +435,4 @@ void FGameplayDebuggerCategory_NavLocalGrid::OnDataPackReplicated(int32 DataPack
 	MarkRenderStateDirty();
 }
 
-#endif // WITH_GAMEPLAY_DEBUGGER
+#endif // WITH_GAMEPLAY_DEBUGGER_MENU
diff --git a/Engine/Source/Runtime/AIModule/Private/GameplayDebugger/GameplayDebuggerCategory_NavLocalGrid.h b/Engine/Source/Runtime/AIModule/Private/GameplayDebugger/GameplayDebuggerCategory_NavLocalGrid.h
deleted file mode 100644
diff --git a/Engine/Source/Runtime/AIModule/Private/GameplayDebugger/GameplayDebuggerCategory_Navmesh.cpp b/Engine/Source/Runtime/AIModule/Private/GameplayDebugger/GameplayDebuggerCategory_Navmesh.cpp
--- a/Engine/Source/Runtime/AIModule/Private/GameplayDebugger/GameplayDebuggerCategory_Navmesh.cpp
+++ b/Engine/Source/Runtime/AIModule/Private/GameplayDebugger/GameplayDebuggerCategory_Navmesh.cpp
@@ -3,7 +3,7 @@
 #include "GameplayDebugger/GameplayDebuggerCategory_Navmesh.h"
 #include "GameFramework/Pawn.h"
 
-#if WITH_GAMEPLAY_DEBUGGER
+#if WITH_GAMEPLAY_DEBUGGER_MENU
 
 #include "NavigationSystem.h"
 #include "GameFramework/PlayerController.h"
@@ -330,4 +330,4 @@ FDebugRenderSceneProxy* FGameplayDebuggerCategory_Navmesh::CreateDebugSceneProxy
 	return NavMeshSceneProxy;
 }
 
-#endif // WITH_GAMEPLAY_DEBUGGER
+#endif // WITH_GAMEPLAY_DEBUGGER_MENU
diff --git a/Engine/Source/Runtime/AIModule/Private/GameplayDebugger/GameplayDebuggerCategory_Perception.cpp b/Engine/Source/Runtime/AIModule/Private/GameplayDebugger/GameplayDebuggerCategory_Perception.cpp
--- a/Engine/Source/Runtime/AIModule/Private/GameplayDebugger/GameplayDebuggerCategory_Perception.cpp
+++ b/Engine/Source/Runtime/AIModule/Private/GameplayDebugger/GameplayDebuggerCategory_Perception.cpp
@@ -2,7 +2,7 @@
 
 #include "GameplayDebugger/GameplayDebuggerCategory_Perception.h"
 
-#if WITH_GAMEPLAY_DEBUGGER
+#if WITH_GAMEPLAY_DEBUGGER_MENU
 
 #include "GameFramework/Pawn.h"
 #include "AIController.h"
@@ -50,4 +50,4 @@ void FGameplayDebuggerCategory_Perception::CollectData(APlayerController* OwnerP
 	}
 }
 
-#endif // WITH_GAMEPLAY_DEBUGGER
+#endif // WITH_GAMEPLAY_DEBUGGER_MENU
diff --git a/Engine/Source/Runtime/AIModule/Private/GameplayDebugger/GameplayDebuggerCategory_Perception.h b/Engine/Source/Runtime/AIModule/Private/GameplayDebugger/GameplayDebuggerCategory_Perception.h
deleted file mode 100644
diff --git a/Engine/Source/Runtime/AIModule/Private/Perception/AIPerceptionComponent.cpp b/Engine/Source/Runtime/AIModule/Private/Perception/AIPerceptionComponent.cpp
--- a/Engine/Source/Runtime/AIModule/Private/Perception/AIPerceptionComponent.cpp
+++ b/Engine/Source/Runtime/AIModule/Private/Perception/AIPerceptionComponent.cpp
@@ -10,10 +10,10 @@
 
 #include UE_INLINE_GENERATED_CPP_BY_NAME(AIPerceptionComponent)
 
-#if WITH_GAMEPLAY_DEBUGGER
+#if WITH_GAMEPLAY_DEBUGGER_MENU
 #include "GameplayDebuggerTypes.h"
 #include "GameplayDebuggerCategory.h"
-#endif
+#endif // WITH_GAMEPLAY_DEBUGGER_MENU
 
 
 
@@ -841,7 +841,7 @@ void UAIPerceptionComponent::SetSenseEnabled(TSubclassOf<UAISense> SenseClass, c
 //----------------------------------------------------------------------//
 // debug
 //----------------------------------------------------------------------//
-#if WITH_GAMEPLAY_DEBUGGER
+#if WITH_GAMEPLAY_DEBUGGER_MENU
 void UAIPerceptionComponent::DescribeSelfToGameplayDebugger(FGameplayDebuggerCategory* DebuggerCategory) const
 {
 	if (DebuggerCategory == nullptr)
@@ -880,7 +880,7 @@ void UAIPerceptionComponent::DescribeSelfToGameplayDebugger(FGameplayDebuggerCat
 		}
 	}
 }
-#endif // WITH_GAMEPLAY_DEBUGGER
+#endif // WITH_GAMEPLAY_DEBUGGER_MENU
 
 #if ENABLE_VISUAL_LOG
 void UAIPerceptionComponent::DescribeSelfToVisLog(FVisualLogEntry* Snapshot) const
diff --git a/Engine/Source/Runtime/AIModule/Private/Perception/AISense.cpp b/Engine/Source/Runtime/AIModule/Private/Perception/AISense.cpp
--- a/Engine/Source/Runtime/AIModule/Private/Perception/AISense.cpp
+++ b/Engine/Source/Runtime/AIModule/Private/Perception/AISense.cpp
@@ -17,10 +17,10 @@
 
 #include UE_INLINE_GENERATED_CPP_BY_NAME(AISense)
 
-#if WITH_GAMEPLAY_DEBUGGER
+#if WITH_GAMEPLAY_DEBUGGER_MENU
 #include "GameplayDebuggerTypes.h"
 #include "GameplayDebuggerCategory.h"
-#endif
+#endif // WITH_GAMEPLAY_DEBUGGER_MENU
 
 const float UAISense::SuspendNextUpdate = FLT_MAX;
 
@@ -140,7 +140,7 @@ FString UAISenseConfig::GetSenseName() const
 	return CachedSenseName;
 }
 
-#if WITH_GAMEPLAY_DEBUGGER
+#if WITH_GAMEPLAY_DEBUGGER_MENU
 static FString DescribeColorHelper(const FColor& Color)
 {
 	const int32 MaxColors = GColorList.GetColorsNum();
@@ -164,7 +164,7 @@ void UAISenseConfig::DescribeSelfToGameplayDebugger(const UAIPerceptionComponent
 			);
 	}
 }
-#endif // WITH_GAMEPLAY_DEBUGGER
+#endif // WITH_GAMEPLAY_DEBUGGER_MENU
 
 //----------------------------------------------------------------------//
 // 
@@ -200,7 +200,7 @@ void UAISenseConfig_Sight::PostEditChangeChainProperty(FPropertyChangedChainEven
 }
 #endif // WITH_EDITOR
 
-#if WITH_GAMEPLAY_DEBUGGER
+#if WITH_GAMEPLAY_DEBUGGER_MENU
 void UAISenseConfig_Sight::DescribeSelfToGameplayDebugger(const UAIPerceptionComponent* PerceptionComponent, FGameplayDebuggerCategory* DebuggerCategory) const
 {
 	if (PerceptionComponent == nullptr || DebuggerCategory == nullptr)
@@ -239,7 +239,7 @@ void UAISenseConfig_Sight::DescribeSelfToGameplayDebugger(const UAIPerceptionCom
 		DebuggerCategory->AddShape(FGameplayDebuggerShape::MakeSegment(RootLocation + (BodyFacing * NearClippingRadius), RootLocation + (RightDirection * NearClippingRadius), SightRangeColor));
 	}
 }
-#endif // WITH_GAMEPLAY_DEBUGGER
+#endif // WITH_GAMEPLAY_DEBUGGER_MENU
 
 //----------------------------------------------------------------------//
 // UAISenseConfig_Hearing
@@ -256,7 +256,7 @@ TSubclassOf<UAISense> UAISenseConfig_Hearing::GetSenseImplementation() const
 	return *Implementation; 
 }
 
-#if WITH_GAMEPLAY_DEBUGGER
+#if WITH_GAMEPLAY_DEBUGGER_MENU
 void UAISenseConfig_Hearing::DescribeSelfToGameplayDebugger(const UAIPerceptionComponent* PerceptionComponent, FGameplayDebuggerCategory* DebuggerCategory) const
 {
 	if (PerceptionComponent == nullptr || DebuggerCategory == nullptr)
@@ -283,7 +283,7 @@ void UAISenseConfig_Hearing::DescribeSelfToGameplayDebugger(const UAIPerceptionC
 		DebuggerCategory->AddShape(FGameplayDebuggerShape::MakeCylinder(OwnerLocation, HearingRange, 25.0f, HearingRangeColor));
 	}
 }
-#endif // WITH_GAMEPLAY_DEBUGGER
+#endif // WITH_GAMEPLAY_DEBUGGER_MENU
 
 //----------------------------------------------------------------------//
 // UAISenseConfig_Prediction
diff --git a/Engine/Source/Runtime/AIModule/Public/GameplayDebugger/GameplayDebuggerCategory_AI.h b/Engine/Source/Runtime/AIModule/Public/GameplayDebugger/GameplayDebuggerCategory_AI.h
new file mode 100644
--- /dev/null
+++ b/Engine/Source/Runtime/AIModule/Public/GameplayDebugger/GameplayDebuggerCategory_AI.h
@@ -0,0 +1,94 @@
+// Copyright Epic Games, Inc. All Rights Reserved.
+
+#pragma once
+
+#include "CoreMinimal.h"
+
+#if WITH_GAMEPLAY_DEBUGGER_MENU
+#include "GameplayDebuggerCategory.h"
+
+class AActor;
+class AAIController;
+class APawn;
+class APlayerController;
+class FDebugRenderSceneProxy;
+class FPoly;
+class UPrimitiveComponent;
+struct FDebugDrawDelegateHelper;
+struct FNavigationPath;
+
+class FGameplayDebuggerCategory_AI : public FGameplayDebuggerCategory
+{
+public:
+	FGameplayDebuggerCategory_AI();
+
+	virtual void CollectData(APlayerController* OwnerPC, AActor* DebugActor) override;
+	virtual void OnDataPackReplicated(int32 DataPackId) override;
+	virtual void DrawData(APlayerController* OwnerPC, FGameplayDebuggerCanvasContext& CanvasContext) override;
+	virtual FDebugRenderSceneProxy* CreateDebugSceneProxy(const UPrimitiveComponent* InComponent, FDebugDrawDelegateHelper*& OutDelegateHelper) override;
+
+	AIMODULE_API static TSharedRef<FGameplayDebuggerCategory> MakeInstance();
+
+protected:
+	void CollectPathData(AAIController* DebugAI);
+	void DrawPath(UWorld* World);
+	void DrawPawnIcons(UWorld* World, AActor* DebugActor, APawn* SkipPawn, FGameplayDebuggerCanvasContext& CanvasContext);
+	void DrawOverheadInfo(AActor& DebugActor, FGameplayDebuggerCanvasContext& CanvasContext);
+
+	struct FRepData
+	{
+		FString ControllerName;
+		FString PawnName;
+		FString MovementBaseInfo;
+		FString MovementModeInfo;
+		FString PathFollowingInfo;
+		int32 NextPathPointIndex;
+		FVector PathGoalLocation;
+		FString CurrentAITask;
+		FString CurrentAIState;
+		FString CurrentAIAssets;
+		FString NavDataInfo;
+		FString MontageInfo;
+		FString TaskQueueInfo;
+		FString TickingTaskInfo;
+		int16 NumTasksInQueue;
+		int16 NumTickingTasks;
+		uint32 bHasController : 1;
+		uint32 bPathHasGoalActor : 1;
+		uint32 bIsUsingPathFollowing : 1;
+		uint32 bIsUsingCharacter : 1;
+		uint32 bIsUsingBehaviorTree : 1;
+		uint32 bIsUsingGameplayTasks : 1;
+
+		FRepData() : NextPathPointIndex(0), NumTasksInQueue(0), NumTickingTasks(0), bHasController(false), bPathHasGoalActor(false),
+			bIsUsingPathFollowing(false), bIsUsingCharacter(false), bIsUsingBehaviorTree(false), bIsUsingGameplayTasks(false)
+		{
+		}
+
+		void Serialize(FArchive& Ar);
+	};
+	FRepData DataPack;
+
+	struct FRepDataPath
+	{
+		struct FPoly
+		{
+			TArray<FVector> Points;
+			FColor Color;
+		};
+
+		TArray<FPoly> PathCorridor;
+		TArray<FVector> PathPoints;
+
+		void Serialize(FArchive& Ar);
+	};	
+	FRepDataPath PathDataPack;
+	int32 PathDataPackId;
+
+private:
+	// do NOT read from this pointer, it's only for validating if path has changed and can be invalid!
+	FNavigationPath* RawLastPath;
+	double LastPathUpdateTime;
+};
+
+#endif // WITH_GAMEPLAY_DEBUGGER_MENU
diff --git a/Engine/Source/Runtime/AIModule/Public/GameplayDebugger/GameplayDebuggerCategory_BehaviorTree.h b/Engine/Source/Runtime/AIModule/Public/GameplayDebugger/GameplayDebuggerCategory_BehaviorTree.h
new file mode 100644
--- /dev/null
+++ b/Engine/Source/Runtime/AIModule/Public/GameplayDebugger/GameplayDebuggerCategory_BehaviorTree.h
@@ -0,0 +1,35 @@
+// Copyright Epic Games, Inc. All Rights Reserved.
+
+#pragma once
+
+#include "CoreMinimal.h"
+#if WITH_GAMEPLAY_DEBUGGER_MENU
+
+#include "GameplayDebuggerCategory.h"
+
+class AActor;
+class APlayerController;
+
+class FGameplayDebuggerCategory_BehaviorTree : public FGameplayDebuggerCategory
+{
+public:
+	FGameplayDebuggerCategory_BehaviorTree();
+
+	virtual void CollectData(APlayerController* OwnerPC, AActor* DebugActor) override;
+	virtual void DrawData(APlayerController* OwnerPC, FGameplayDebuggerCanvasContext& CanvasContext) override;
+
+	AIMODULE_API static TSharedRef<FGameplayDebuggerCategory> MakeInstance();
+
+protected:
+	struct FRepData
+	{
+		FString CompName;
+		FString TreeDesc;
+		FString BlackboardDesc;
+
+		void Serialize(FArchive& Ar);
+	};
+	FRepData DataPack;
+};
+
+#endif // WITH_GAMEPLAY_DEBUGGER_MENU
\ No newline at end of file
diff --git a/Engine/Source/Runtime/AIModule/Public/GameplayDebugger/GameplayDebuggerCategory_EQS.h b/Engine/Source/Runtime/AIModule/Public/GameplayDebugger/GameplayDebuggerCategory_EQS.h
new file mode 100644
--- /dev/null
+++ b/Engine/Source/Runtime/AIModule/Public/GameplayDebugger/GameplayDebuggerCategory_EQS.h
@@ -0,0 +1,54 @@
+// Copyright Epic Games, Inc. All Rights Reserved.
+
+#pragma once
+
+#include "CoreMinimal.h"
+
+#if WITH_GAMEPLAY_DEBUGGER_MENU
+#include "GameplayDebuggerCategory.h"
+#include "EnvironmentQuery/EnvQueryTypes.h"
+#include "EnvironmentQuery/EnvQueryDebugHelpers.h"
+
+class APlayerController;
+class UPrimitiveComponent;
+
+class FGameplayDebuggerCategory_EQS : public FGameplayDebuggerCategory
+{
+public:
+	FGameplayDebuggerCategory_EQS();
+
+	virtual void CollectData(APlayerController* OwnerPC, AActor* DebugActor) override;
+	virtual void OnDataPackReplicated(int32 DataPackId) override;
+	virtual void DrawData(APlayerController* OwnerPC, FGameplayDebuggerCanvasContext& CanvasContext) override;
+	virtual FDebugRenderSceneProxy* CreateDebugSceneProxy(const UPrimitiveComponent* InComponent, FDebugDrawDelegateHelper*& OutDelegateHelper) override;
+
+	AIMODULE_API static TSharedRef<FGameplayDebuggerCategory> MakeInstance();
+
+protected:
+#if USE_EQS_DEBUGGER
+	struct FRepData
+	{
+		TArray<EQSDebug::FQueryData> QueryDebugData;
+		
+		void Serialize(FArchive& Ar);
+	};
+	FRepData DataPack;
+
+	int32 DrawLookedAtItem(const EQSDebug::FQueryData& QueryData, APlayerController* OwnerPC, FGameplayDebuggerCanvasContext& CanvasContext) const;
+	void DrawDetailedItemTable(const EQSDebug::FQueryData& QueryData, int32 LookedAtItemIndex, FGameplayDebuggerCanvasContext& CanvasContext) const;
+	void DrawDetailedItemRow(const EQSDebug::FItemData& ItemData, const TArray<uint8>& TestRelevancy, float MaxScore, FGameplayDebuggerCanvasContext& CanvasContext) const;
+#endif
+
+	void CycleShownQueries();
+	void ToggleDetailView();
+
+	uint32 bDrawLabels : 1;
+	uint32 bDrawFailedItems : 1;
+	uint32 bShowDetails : 1;
+
+	int32 MaxItemTableRows;
+	int32 MaxQueries;
+	int32 ShownQueryIndex;
+};
+
+#endif // WITH_GAMEPLAY_DEBUGGER_MENU
diff --git a/Engine/Source/Runtime/AIModule/Public/GameplayDebugger/GameplayDebuggerCategory_NavLocalGrid.h b/Engine/Source/Runtime/AIModule/Public/GameplayDebugger/GameplayDebuggerCategory_NavLocalGrid.h
new file mode 100644
--- /dev/null
+++ b/Engine/Source/Runtime/AIModule/Public/GameplayDebugger/GameplayDebuggerCategory_NavLocalGrid.h
@@ -0,0 +1,56 @@
+// Copyright Epic Games, Inc. All Rights Reserved.
+
+#pragma once
+
+#include "CoreMinimal.h"
+
+#if WITH_GAMEPLAY_DEBUGGER_MENU
+#include "GameplayDebuggerCategory.h"
+
+class APlayerController;
+
+class FGameplayDebuggerCategory_NavLocalGrid : public FGameplayDebuggerCategory
+{
+public:
+	FGameplayDebuggerCategory_NavLocalGrid();
+
+	virtual void CollectData(APlayerController* OwnerPC, AActor* DebugActor) override;
+	virtual void DrawData(APlayerController* OwnerPC, FGameplayDebuggerCanvasContext& CanvasContext) override;
+	virtual FDebugRenderSceneProxy* CreateDebugSceneProxy(const UPrimitiveComponent* InComponent, FDebugDrawDelegateHelper*& OutDelegateHelper) override;
+	virtual void OnDataPackReplicated(int32 DataPackId) override;
+
+	AIMODULE_API static TSharedRef<FGameplayDebuggerCategory> MakeInstance();
+
+	struct FRepData
+	{
+		struct FGridData
+		{
+			FBox Bounds;
+			float CellSize;
+			int32 NumRows;
+			int32 NumCols;
+			TArray<uint8> Cells;
+		};
+
+		TArray<FGridData> Grids;
+		int32 NumSources;
+		int32 VersionNum;
+
+		FRepData() : NumSources(0), VersionNum(0) {}
+		void Serialize(FArchive& Ar);
+	};
+
+	struct FRepAgentData
+	{
+		TArray<int32> PathCells;
+		int32 GridIdx;
+
+		void Serialize(FArchive& Ar);
+	};
+
+protected:
+	FRepData DataPack;
+	FRepAgentData AgentDataPack;
+};
+
+#endif // WITH_GAMEPLAY_DEBUGGER_MENU
diff --git a/Engine/Source/Runtime/AIModule/Public/GameplayDebugger/GameplayDebuggerCategory_Navmesh.h b/Engine/Source/Runtime/AIModule/Public/GameplayDebugger/GameplayDebuggerCategory_Navmesh.h
--- a/Engine/Source/Runtime/AIModule/Public/GameplayDebugger/GameplayDebuggerCategory_Navmesh.h
+++ b/Engine/Source/Runtime/AIModule/Public/GameplayDebugger/GameplayDebuggerCategory_Navmesh.h
@@ -2,9 +2,9 @@
 
 #pragma once
 
-#if WITH_GAMEPLAY_DEBUGGER
-
 #include "CoreMinimal.h"
+
+#if WITH_GAMEPLAY_DEBUGGER_MENU
 #include "GameplayDebuggerCategory.h"
 #include "NavMesh/NavMeshRenderingComponent.h"
 
@@ -63,4 +63,4 @@ protected:
 	TWeakObjectPtr<const APawn> PrevDebugActorReference;
 };
 
-#endif // WITH_GAMEPLAY_DEBUGGER
+#endif // WITH_GAMEPLAY_DEBUGGER_MENU
diff --git a/Engine/Source/Runtime/AIModule/Public/GameplayDebugger/GameplayDebuggerCategory_Perception.h b/Engine/Source/Runtime/AIModule/Public/GameplayDebugger/GameplayDebuggerCategory_Perception.h
new file mode 100644
--- /dev/null
+++ b/Engine/Source/Runtime/AIModule/Public/GameplayDebugger/GameplayDebuggerCategory_Perception.h
@@ -0,0 +1,23 @@
+// Copyright Epic Games, Inc. All Rights Reserved.
+
+#pragma once
+
+#include "CoreMinimal.h"
+
+#if WITH_GAMEPLAY_DEBUGGER_MENU
+#include "GameplayDebuggerCategory.h"
+
+class AActor;
+class APlayerController;
+
+class FGameplayDebuggerCategory_Perception : public FGameplayDebuggerCategory
+{
+public:
+	FGameplayDebuggerCategory_Perception();
+
+	virtual void CollectData(APlayerController* OwnerPC, AActor* DebugActor) override;
+
+	AIMODULE_API static TSharedRef<FGameplayDebuggerCategory> MakeInstance();
+};
+
+#endif // WITH_GAMEPLAY_DEBUGGER_MENU
diff --git a/Engine/Source/Runtime/GameplayDebugger/GameplayDebugger.Build.cs b/Engine/Source/Runtime/GameplayDebugger/GameplayDebugger.Build.cs
--- a/Engine/Source/Runtime/GameplayDebugger/GameplayDebugger.Build.cs
+++ b/Engine/Source/Runtime/GameplayDebugger/GameplayDebugger.Build.cs
@@ -55,11 +55,20 @@ namespace UnrealBuildTool.Rules
 			{
 				PublicDefinitions.Add("WITH_GAMEPLAY_DEBUGGER_CORE=1");
 				PublicDefinitions.Add("WITH_GAMEPLAY_DEBUGGER=" + (Target.bUseGameplayDebugger ? 1 : 0));
+				if (Target.Configuration != UnrealTargetConfiguration.Shipping)
+				{
+					PublicDefinitions.Add("WITH_GAMEPLAY_DEBUGGER_MENU=1");
+				}
+				else
+				{
+					PublicDefinitions.Add("WITH_GAMEPLAY_DEBUGGER_MENU=0");
+				}				
 			}
 			else
 			{
 				PublicDefinitions.Add("WITH_GAMEPLAY_DEBUGGER_CORE=0");
 				PublicDefinitions.Add("WITH_GAMEPLAY_DEBUGGER=0");
+				PublicDefinitions.Add("WITH_GAMEPLAY_DEBUGGER_MENU=0");
 			}
 		}
 	}
diff --git a/Engine/Source/Runtime/GameplayDebugger/Private/GameplayDebuggerLocalController.cpp b/Engine/Source/Runtime/GameplayDebugger/Private/GameplayDebuggerLocalController.cpp
--- a/Engine/Source/Runtime/GameplayDebugger/Private/GameplayDebuggerLocalController.cpp
+++ b/Engine/Source/Runtime/GameplayDebugger/Private/GameplayDebuggerLocalController.cpp
@@ -35,6 +35,12 @@
 #include "Editor.h"
 #endif // WITH_EDITOR
 
+#if WITH_GAMEPLAY_DEBUGGER
+bool UGameplayDebuggerLocalController::bConsoleCommandsEnabled = true;
+#elif WITH_GAMEPLAY_DEBUGGER_MENU
+bool UGameplayDebuggerLocalController::bConsoleCommandsEnabled = false;
+#endif // WITH_GAMEPLAY_DEBUGGER_MENU
+
 UGameplayDebuggerLocalController::UGameplayDebuggerLocalController(const FObjectInitializer& ObjectInitializer) : Super(ObjectInitializer)
 {
 	bSimulateMode = false;
@@ -48,14 +54,14 @@ UGameplayDebuggerLocalController::UGameplayDebuggerLocalController(const FObject
 	bActivateOnPIEEnd = false;
 #endif // WITH_EDITOR
 
-#if WITH_GAMEPLAY_DEBUGGER
+#if WITH_GAMEPLAY_DEBUGGER_MENU
 	ActiveRowIdx = 0;
 	if (HasAnyFlags(RF_ClassDefaultObject) == false)
 	{
 		HUDFont = NewObject<UFont>(this, TEXT("HUDFont"), RF_NoFlags, GEngine->GetSmallFont());
 		HUDFont->LegacyFontSize = UGameplayDebuggerUserSettings::GetFontSize(); //FGameplayDebuggerTweakables::FontSize;
 	}
-#endif // WITH_GAMEPLAY_DEBUGGER
+#endif // WITH_GAMEPLAY_DEBUGGER_MENU
 }
 
 void UGameplayDebuggerLocalController::Initialize(AGameplayDebuggerCategoryReplicator& Replicator, AGameplayDebuggerPlayerManager& Manager)
@@ -64,7 +70,7 @@ void UGameplayDebuggerLocalController::Initialize(AGameplayDebuggerCategoryRepli
 	CachedPlayerManager = &Manager;
 	bSimulateMode = FGameplayDebuggerAddonBase::IsSimulateInEditor() || Replicator.IsEditorWorldReplicator();
 
-#if WITH_GAMEPLAY_DEBUGGER
+#if WITH_GAMEPLAY_DEBUGGER_MENU
 	UDebugDrawService::Register(bSimulateMode ? TEXT("DebugAI") : TEXT("Game"), FDebugDrawDelegate::CreateUObject(this, &UGameplayDebuggerLocalController::OnDebugDraw));
 
 #if WITH_EDITOR
@@ -124,7 +130,7 @@ void UGameplayDebuggerLocalController::Initialize(AGameplayDebuggerCategoryRepli
 	PaddingBottom = SettingsCDO->DebugCanvasPaddingBottom;
 
 	bEnableTextShadow = SettingsCDO->bDebugCanvasEnableTextShadow;
-#endif // WITH_GAMEPLAY_DEBUGGER
+#endif // WITH_GAMEPLAY_DEBUGGER_MENU
 
 	FGameplayDebuggerAddonManager& AddonManager = FGameplayDebuggerAddonManager::GetCurrent();
 	AddonManager.OnCategoriesChanged.AddUObject(this, &UGameplayDebuggerLocalController::OnCategoriesChanged);
@@ -185,7 +191,7 @@ void UGameplayDebuggerLocalController::OnCategoriesChanged()
 	DataPackMap.Reset();
 }
 
-#if WITH_GAMEPLAY_DEBUGGER
+#if WITH_GAMEPLAY_DEBUGGER_MENU
 void UGameplayDebuggerLocalController::OnDebugDraw(class UCanvas* Canvas, class APlayerController* PC)
 {
 	if (CachedReplicator && CachedReplicator->IsEnabled() && bDebugDrawEnabled)
@@ -901,6 +907,11 @@ struct FGameplayDebuggerConsoleCommands
 private:
 	static UGameplayDebuggerLocalController* GetController(UWorld* InWorld)
 	{
+		if (!UGameplayDebuggerLocalController::bConsoleCommandsEnabled)
+		{
+			return nullptr;
+		}
+
 		UGameplayDebuggerLocalController* Controller = nullptr;
 
 		APlayerController* LocalPC = GEngine->GetFirstLocalPlayerController(InWorld);
@@ -1105,4 +1116,4 @@ FAutoConsoleCommandWithWorldAndArgs FGameplayDebuggerConsoleCommands::SetFontSiz
 	FConsoleCommandWithWorldAndArgsDelegate::CreateStatic(&FGameplayDebuggerConsoleCommands::SetFontSize)
 );
 
-#endif // WITH_GAMEPLAY_DEBUGGER
\ No newline at end of file
+#endif // WITH_GAMEPLAY_DEBUGGER_MENU
\ No newline at end of file
diff --git a/Engine/Source/Runtime/GameplayDebugger/Private/GameplayDebuggerPlayerManager.cpp b/Engine/Source/Runtime/GameplayDebugger/Private/GameplayDebuggerPlayerManager.cpp
--- a/Engine/Source/Runtime/GameplayDebugger/Private/GameplayDebuggerPlayerManager.cpp
+++ b/Engine/Source/Runtime/GameplayDebugger/Private/GameplayDebuggerPlayerManager.cpp
@@ -197,7 +197,7 @@ void AGameplayDebuggerPlayerManager::RegisterReplicator(AGameplayDebuggerCategor
 	FGameplayDebuggerPlayerData NewData;
 	NewData.Replicator = &Replicator;
 
-#if WITH_GAMEPLAY_DEBUGGER
+#if WITH_GAMEPLAY_DEBUGGER_MENU
 	if (bIsLocal)
 	{
 		APlayerController* OwnerPC = Replicator.GetReplicationOwner();
@@ -219,14 +219,14 @@ void AGameplayDebuggerPlayerManager::RegisterReplicator(AGameplayDebuggerCategor
 		NewData.Controller = nullptr;
 		NewData.InputComponent = nullptr;
 	}
-#endif // WITH_GAMEPLAY_DEBUGGER
+#endif // WITH_GAMEPLAY_DEBUGGER_MENU
 
 	PlayerData.Add(NewData);
 }
 
 void AGameplayDebuggerPlayerManager::RefreshInputBindings(AGameplayDebuggerCategoryReplicator& Replicator)
 {
-#if WITH_GAMEPLAY_DEBUGGER
+#if WITH_GAMEPLAY_DEBUGGER_MENU
 	for (int32 Idx = 0; Idx < PlayerData.Num(); Idx++)
 	{
 		FGameplayDebuggerPlayerData& TestData = PlayerData[Idx];
@@ -239,7 +239,7 @@ void AGameplayDebuggerPlayerManager::RefreshInputBindings(AGameplayDebuggerCateg
 			TestData.Controller->BindInput(*TestData.InputComponent);
 		}
 	}
-#endif // WITH_GAMEPLAY_DEBUGGER
+#endif // WITH_GAMEPLAY_DEBUGGER_MENU
 }
 
 AGameplayDebuggerCategoryReplicator* AGameplayDebuggerPlayerManager::GetReplicator(const APlayerController& OwnerPC) const
diff --git a/Engine/Source/Runtime/GameplayDebugger/Public/GameplayDebuggerCategory.h b/Engine/Source/Runtime/GameplayDebugger/Public/GameplayDebuggerCategory.h
--- a/Engine/Source/Runtime/GameplayDebugger/Public/GameplayDebuggerCategory.h
+++ b/Engine/Source/Runtime/GameplayDebugger/Public/GameplayDebuggerCategory.h
@@ -8,7 +8,7 @@
 // show server's state on client.
 // 
 // It should be compiled and used only when module is included, so every category class
-// needs be placed in #if WITH_GAMEPLAY_DEBUGGER block.
+// needs be placed in #if WITH_GAMEPLAY_DEBUGGER block (or WITH_GAMEPLAY_DEBUGGER_MENU when using WITH_GAMEPLAY_DEBUGGER_CORE).
 // 
 // 
 // Server side category :
diff --git a/Engine/Source/Runtime/GameplayDebugger/Public/GameplayDebuggerLocalController.h b/Engine/Source/Runtime/GameplayDebugger/Public/GameplayDebuggerLocalController.h
--- a/Engine/Source/Runtime/GameplayDebugger/Public/GameplayDebuggerLocalController.h
+++ b/Engine/Source/Runtime/GameplayDebugger/Public/GameplayDebuggerLocalController.h
@@ -30,7 +30,7 @@ class UGameplayDebuggerLocalController : public UObject
 	/** remove from world */
 	void Cleanup();
 
-#if WITH_GAMEPLAY_DEBUGGER
+#if WITH_GAMEPLAY_DEBUGGER_MENU
 	/** Whether the Debug Draw rendering is allowed for this replicator */
 	void SetDebugDrawEnabled(const bool bInDebugDrawEnabled) { bDebugDrawEnabled = bInDebugDrawEnabled; }
 
@@ -42,8 +42,9 @@ class UGameplayDebuggerLocalController : public UObject
 
 	/** checks if key is bound by any action */
 	bool GAMEPLAYDEBUGGER_API IsKeyBound(const FName KeyName) const;
-#endif // WITH_GAMEPLAY_DEBUGGER
 
+	static void SetConsoleCommandsEnabled(const bool bEnabled) { bConsoleCommandsEnabled = bEnabled; }
+#endif // WITH_GAMEPLAY_DEBUGGER_MENU	
 protected:
 	UPROPERTY()
 	TObjectPtr<AGameplayDebuggerCategoryReplicator> CachedReplicator;
@@ -74,7 +75,7 @@ protected:
 	int32 NumCategorySlots;
 	int32 NumCategories;
 
-#if WITH_GAMEPLAY_DEBUGGER
+#if WITH_GAMEPLAY_DEBUGGER_MENU
 	friend struct FGameplayDebuggerConsoleCommands;
 
 	TSet<FName> UsedBindings;
@@ -96,6 +97,8 @@ protected:
 	FTimerHandle SelectActorTickHandle;
 
 	bool bDebugDrawEnabled = true;
+	
+	static bool bConsoleCommandsEnabled;
 
 	void OnActivationPressed();
 	void OnActivationReleased();
@@ -148,7 +151,7 @@ protected:
 	
 	/** build DataPackMap for replication details */
 	void RebuildDataPackMap();
-#endif // WITH_GAMEPLAY_DEBUGGER
+#endif // WITH_GAMEPLAY_DEBUGGER_MENU
 
 	/** called when known category set has changed */
 	void OnCategoriesChanged();
