- Defined WITH_GAMEPLAY_DEBUGER_MENU to compile the gameplay debugger categories menus either when WITH_GAMEPLAY_DEBUGGER is defined or when WITH_GAMEPLAY_DEBUGGER_CORE is defined but except in Shipping.
- Debugger Input bindings are automatically registered only when WITH_GAMEPLAY_DEBUGGER is enabled.

diff --git a/Engine/Source/Programs/UnrealBuildTool/Configuration/ModuleRules.cs b/Engine/Source/Programs/UnrealBuildTool/Configuration/ModuleRules.cs
--- a/Engine/Source/Programs/UnrealBuildTool/Configuration/ModuleRules.cs
+++ b/Engine/Source/Programs/UnrealBuildTool/Configuration/ModuleRules.cs
@@ -1336,7 +1336,7 @@ namespace UnrealBuildTool
 			{
 				PublicDefinitions.Add("WITH_GAMEPLAY_DEBUGGER_CORE=1");
 				PublicDefinitions.Add("WITH_GAMEPLAY_DEBUGGER=" + (Target.bUseGameplayDebugger ? 1 : 0));
-				if (Target.Configuration != UnrealTargetConfiguration.Shipping || Target.bBuildDeveloperTools)
+				if (Target.bUseGameplayDebugger || (Target.bUseGameplayDebuggerCore && Target.Configuration != UnrealTargetConfiguration.Shipping))
 				{
 					PublicDefinitions.Add("WITH_GAMEPLAY_DEBUGGER_MENU=1");
 				}
diff --git a/Engine/Source/Runtime/GameplayDebugger/GameplayDebugger.Build.cs b/Engine/Source/Runtime/GameplayDebugger/GameplayDebugger.Build.cs
--- a/Engine/Source/Runtime/GameplayDebugger/GameplayDebugger.Build.cs
+++ b/Engine/Source/Runtime/GameplayDebugger/GameplayDebugger.Build.cs
@@ -55,14 +55,14 @@ namespace UnrealBuildTool.Rules
 			{
 				PublicDefinitions.Add("WITH_GAMEPLAY_DEBUGGER_CORE=1");
 				PublicDefinitions.Add("WITH_GAMEPLAY_DEBUGGER=" + (Target.bUseGameplayDebugger ? 1 : 0));
-				if (Target.Configuration != UnrealTargetConfiguration.Shipping || Target.bBuildDeveloperTools)
+				if (Target.bUseGameplayDebugger || (Target.bUseGameplayDebuggerCore && Target.Configuration != UnrealTargetConfiguration.Shipping))
 				{
 					PublicDefinitions.Add("WITH_GAMEPLAY_DEBUGGER_MENU=1");
 				}
 				else
 				{
 					PublicDefinitions.Add("WITH_GAMEPLAY_DEBUGGER_MENU=0");
-				}				
+				}
 			}
 			else
 			{
diff --git a/Engine/Source/Runtime/GameplayDebugger/Private/GameplayDebuggerLocalController.cpp b/Engine/Source/Runtime/GameplayDebugger/Private/GameplayDebuggerLocalController.cpp
--- a/Engine/Source/Runtime/GameplayDebugger/Private/GameplayDebuggerLocalController.cpp
+++ b/Engine/Source/Runtime/GameplayDebugger/Private/GameplayDebuggerLocalController.cpp
@@ -37,9 +37,9 @@
 
 #if WITH_GAMEPLAY_DEBUGGER
 bool UGameplayDebuggerLocalController::bConsoleCommandsEnabled = true;
-#elif WITH_GAMEPLAY_DEBUGGER_MENU
+#else
 bool UGameplayDebuggerLocalController::bConsoleCommandsEnabled = false;
-#endif // WITH_GAMEPLAY_DEBUGGER_MENU
+#endif
 
 UGameplayDebuggerLocalController::UGameplayDebuggerLocalController(const FObjectInitializer& ObjectInitializer) : Super(ObjectInitializer)
 {
