GameplayDebugger: Added rectangle, polyline and segment list shapes.

diff --git a/Engine/Source/Runtime/GameplayDebugger/Private/GameplayDebuggerTypes.cpp b/Engine/Source/Runtime/GameplayDebugger/Private/GameplayDebuggerTypes.cpp
--- a/Engine/Source/Runtime/GameplayDebugger/Private/GameplayDebuggerTypes.cpp
+++ b/Engine/Source/Runtime/GameplayDebugger/Private/GameplayDebuggerTypes.cpp
@@ -113,10 +113,32 @@ FGameplayDebuggerShape FGameplayDebuggerShape::MakeCircle(const FVector& Center,
 }
 
 FGameplayDebuggerShape FGameplayDebuggerShape::MakeCircle(const FVector& Center, const FVector& Up, const float Radius, const float Thickness, const FColor& Color, const FString& Description)
+{
+	const FMatrix Axes = FRotationMatrix::MakeFromX(Up);
+	
+	FGameplayDebuggerShape NewElement;
+	NewElement.ShapeData.Add(Center);
+	NewElement.ShapeData.Add(Axes.GetUnitAxis(EAxis::Y));
+	NewElement.ShapeData.Add(Axes.GetUnitAxis(EAxis::Z));
+	NewElement.ShapeData.Add(FVector(Radius, Thickness, 0));
+	NewElement.Color = Color;
+	NewElement.Description = Description;
+	NewElement.Type = EGameplayDebuggerShape::Circle;
+
+	return NewElement;
+}
+
+FGameplayDebuggerShape FGameplayDebuggerShape::MakeCircle(const FVector& Center, const FVector& WidthAxis, const FVector& HeightAxis, const float Radius, const FColor& Color, const FString& Description)
+{
+	return MakeCircle(Center, WidthAxis, HeightAxis, Radius, 1.0f, Color, Description);
+}
+
+FGameplayDebuggerShape FGameplayDebuggerShape::MakeCircle(const FVector& Center, const FVector& WidthAxis, const FVector& HeightAxis, const float Radius, const float Thickness, const FColor& Color, const FString& Description)
 {
 	FGameplayDebuggerShape NewElement;
 	NewElement.ShapeData.Add(Center);
-	NewElement.ShapeData.Add(Up);
+	NewElement.ShapeData.Add(WidthAxis);
+	NewElement.ShapeData.Add(HeightAxis);
 	NewElement.ShapeData.Add(FVector(Radius, Thickness, 0));
 	NewElement.Color = Color;
 	NewElement.Description = Description;
@@ -125,6 +147,25 @@ FGameplayDebuggerShape FGameplayDebuggerShape::MakeCircle(const FVector& Center,
 	return NewElement;
 }
 
+FGameplayDebuggerShape FGameplayDebuggerShape::MakeRectangle(const FVector& Center, const FVector& WidthAxis, const FVector& HeightAxis, const float Width, const float Height, const FColor& Color, const FString& Description)
+{
+	return MakeRectangle(Center, WidthAxis, HeightAxis, Width, Height, 1.0f, Color, Description);
+}
+
+FGameplayDebuggerShape FGameplayDebuggerShape::MakeRectangle(const FVector& Center, const FVector& WidthAxis, const FVector& HeightAxis, const float Width, const float Height, const float Thickness, const FColor& Color, const FString& Description)
+{
+	FGameplayDebuggerShape NewElement;
+	NewElement.ShapeData.Add(Center);
+	NewElement.ShapeData.Add(WidthAxis * Width * 0.5);
+	NewElement.ShapeData.Add(HeightAxis * Height * 0.5);
+	NewElement.ShapeData.Add(FVector(Thickness, 0, 0));
+	NewElement.Color = Color;
+	NewElement.Description = Description;
+	NewElement.Type = EGameplayDebuggerShape::Rectangle;
+
+	return NewElement;
+}
+
 FGameplayDebuggerShape FGameplayDebuggerShape::MakeCapsule(const FVector& Center, const float Radius, const float HalfHeight, const FColor& Color, const FString& Description)
 {
 	FGameplayDebuggerShape NewElement;
@@ -137,7 +178,7 @@ FGameplayDebuggerShape FGameplayDebuggerShape::MakeCapsule(const FVector& Center
 	return NewElement;
 }
 
-FGameplayDebuggerShape FGameplayDebuggerShape::MakePolygon(const TArray<FVector>& Verts, const FColor& Color, const FString& Description)
+FGameplayDebuggerShape FGameplayDebuggerShape::MakePolygon(TConstArrayView<FVector> Verts, const FColor& Color, const FString& Description)
 {
 	FGameplayDebuggerShape NewElement;
 	NewElement.ShapeData = Verts;
@@ -147,6 +188,38 @@ FGameplayDebuggerShape FGameplayDebuggerShape::MakePolygon(const TArray<FVector>
 	return NewElement;
 }
 
+FGameplayDebuggerShape FGameplayDebuggerShape::MakePolyline(const TConstArrayView<FVector> Verts, const FColor& Color, const FString& Description)
+{
+	return MakePolyline(Verts, 1.0f, Color, Description);
+}
+
+FGameplayDebuggerShape FGameplayDebuggerShape::MakePolyline(const TConstArrayView<FVector> Verts, const float Thickness, const FColor& Color, const FString& Description)
+{
+	FGameplayDebuggerShape NewElement;
+	NewElement.ShapeData = Verts;
+	NewElement.ShapeData.Add(FVector(Thickness, 0, 0));
+	NewElement.Color = Color;
+	NewElement.Description = Description;
+	NewElement.Type = EGameplayDebuggerShape::Polyline;
+	return NewElement;
+}
+
+FGameplayDebuggerShape FGameplayDebuggerShape::MakeSegmentList(TConstArrayView<FVector> Verts, const float Thickness, const FColor& Color, const FString& Description)
+{
+	FGameplayDebuggerShape NewElement;
+	NewElement.ShapeData = Verts;
+	NewElement.ShapeData.Add(FVector(Thickness, 0, 0));
+	NewElement.Color = Color;
+	NewElement.Description = Description;
+	NewElement.Type = EGameplayDebuggerShape::Segment;
+	return NewElement;
+}
+
+FGameplayDebuggerShape FGameplayDebuggerShape::MakeSegmentList(TConstArrayView<FVector> Verts, const FColor& Color, const FString& Description)
+{
+	return MakeSegmentList(Verts, 1.0f, Color, Description);
+}
+
 void FGameplayDebuggerShape::Draw(UWorld* World, FGameplayDebuggerCanvasContext& Context)
 {
 	constexpr bool bPersistent = false;
@@ -165,10 +238,20 @@ void FGameplayDebuggerShape::Draw(UWorld* World, FGameplayDebuggerCanvasContext&
 		break;
 
 	case EGameplayDebuggerShape::Segment:
-		if (ShapeData.Num() == 3 && ShapeData[2].X > 0)
+		if (ShapeData.Num() >= 3 && ShapeData.Last().X > 0)
 		{
-			DrawDebugLine(World, ShapeData[0], ShapeData[1], Color, bPersistent, LifeTime, DepthPriority, static_cast<float>(ShapeData[2].X));
-			DescLocation = (ShapeData[0] + ShapeData[1]) * 0.5f;
+			const int32 NumVertices = (ShapeData.Num() - 1) & ~1; // Expect 2 vertices per line.
+			const float Thickness = static_cast<float>(ShapeData.Last().X);
+			DescLocation = FVector::ZeroVector;
+			for (int32 Index = 0; Index < NumVertices; Index += 2)
+			{
+				DrawDebugLine(World, ShapeData[Index + 0], ShapeData[Index + 1], Color, bPersistent, LifeTime, DepthPriority, Thickness);
+				DescLocation += ShapeData[Index + 0] + ShapeData[Index + 1];
+			}
+			if (NumVertices > 0)
+			{
+				DescLocation /= NumVertices; 
+			}
 		}
 		break;
 
@@ -191,7 +274,7 @@ void FGameplayDebuggerShape::Draw(UWorld* World, FGameplayDebuggerCanvasContext&
 	case EGameplayDebuggerShape::Cone:
 		if (ShapeData.Num() == 3 && ShapeData[2].X > 0)
 		{
-			const float DefaultConeAngle = 0.25f; // ~ 15 degrees
+			constexpr float DefaultConeAngle = 0.25f; // ~ 15 degrees
 			DrawDebugCone(World, ShapeData[0], ShapeData[1], static_cast<float>(ShapeData[2].X), DefaultConeAngle, DefaultConeAngle, 16, Color, bPersistent, LifeTime, DepthPriority);
 			DescLocation = ShapeData[0];
 		}
@@ -206,10 +289,9 @@ void FGameplayDebuggerShape::Draw(UWorld* World, FGameplayDebuggerCanvasContext&
 		break;
 
 	case EGameplayDebuggerShape::Circle:
-		if (ShapeData.Num() == 3)
+		if (ShapeData.Num() == 4)
 		{
-			const FMatrix Axes = FRotationMatrix::MakeFromX(ShapeData[1]);
-			DrawDebugCircle(World, ShapeData[0], static_cast<float>(ShapeData[2].X), 32, Color, bPersistent, LifeTime, DepthPriority, ShapeData[2].Y, Axes.GetUnitAxis(EAxis::Y), Axes.GetUnitAxis(EAxis::Z), false);
+			DrawDebugCircle(World, ShapeData[0], static_cast<float>(ShapeData[3].X), 32, Color, bPersistent, LifeTime, DepthPriority, ShapeData[3].Y, ShapeData[1], ShapeData[2], false);
 			DescLocation = ShapeData[0];
 		}
 		break;
@@ -222,6 +304,24 @@ void FGameplayDebuggerShape::Draw(UWorld* World, FGameplayDebuggerCanvasContext&
 		}
 		break;
 
+	case EGameplayDebuggerShape::Polyline:
+		if (ShapeData.Num() >= 3 && ShapeData.Last().X > 0)
+		{
+			const int32 NumVertices = (ShapeData.Num() - 1);
+			const FVector::FReal Thickness = ShapeData.Last().X;
+			DescLocation = ShapeData[0];
+			for (int32 Index = 0; Index < NumVertices - 1; Index++)
+			{
+				DrawDebugLine(World, ShapeData[Index + 0], ShapeData[Index + 1], Color, bPersistent, LifeTime, DepthPriority, Thickness);
+				DescLocation += ShapeData[Index + 1];
+			}
+			if (NumVertices > 0)
+			{
+				DescLocation /= NumVertices; 
+			}
+		}
+		break;
+		
 	case EGameplayDebuggerShape::Polygon:
 		if (ShapeData.Num() > 0)
 		{
@@ -238,6 +338,24 @@ void FGameplayDebuggerShape::Draw(UWorld* World, FGameplayDebuggerCanvasContext&
 		}
 		break;
 
+	case EGameplayDebuggerShape::Rectangle:
+		if (ShapeData.Num() == 4)
+		{
+			const FVector& Center = ShapeData[0];
+			const FVector& WidthAxis = ShapeData[1];
+			const FVector& HeightAxis = ShapeData[2];
+			const FVector P0 = Center - WidthAxis + HeightAxis;
+			const FVector P1 = Center + WidthAxis + HeightAxis;
+			const FVector P2 = Center + WidthAxis - HeightAxis;
+			const FVector P3 = Center - WidthAxis - HeightAxis;
+			const FVector::FReal Thickness = ShapeData[3].X;
+			DrawDebugLine(World, P0, P1 , Color, bPersistent, LifeTime, DepthPriority, Thickness);
+			DrawDebugLine(World, P1, P2 , Color, bPersistent, LifeTime, DepthPriority, Thickness);
+			DrawDebugLine(World, P2, P3 , Color, bPersistent, LifeTime, DepthPriority, Thickness);
+			DrawDebugLine(World, P3, P0 , Color, bPersistent, LifeTime, DepthPriority, Thickness);
+		}
+		break;
+		
 	default:
 		break;
 	}
diff --git a/Engine/Source/Runtime/GameplayDebugger/Public/GameplayDebuggerTypes.h b/Engine/Source/Runtime/GameplayDebugger/Public/GameplayDebuggerTypes.h
--- a/Engine/Source/Runtime/GameplayDebugger/Public/GameplayDebuggerTypes.h
+++ b/Engine/Source/Runtime/GameplayDebugger/Public/GameplayDebuggerTypes.h
@@ -169,8 +169,10 @@ enum class EGameplayDebuggerShape : uint8
 	Cone,
 	Cylinder,
 	Circle,
+	Rectangle,
 	Capsule,
 	Polygon,
+	Polyline,
 	Arrow,
 };
 
@@ -211,8 +213,16 @@ struct GAMEPLAYDEBUGGER_API FGameplayDebuggerShape
 	static FGameplayDebuggerShape MakeCylinder(const FVector& Center, const float Radius, const float HalfHeight, const FColor& Color, const FString& Description = FString());
 	static FGameplayDebuggerShape MakeCircle(const FVector& Center, const FVector& Up, const float Radius, const FColor& Color, const FString& Description = FString());
 	static FGameplayDebuggerShape MakeCircle(const FVector& Center, const FVector& Up, const float Radius, const float Thickness, const FColor& Color, const FString& Description = FString());
+	static FGameplayDebuggerShape MakeCircle(const FVector& Center, const FVector& WidthAxis, const FVector& HeightAxis, const float Radius, const FColor& Color, const FString& Description = FString());
+	static FGameplayDebuggerShape MakeCircle(const FVector& Center, const FVector& WidthAxis, const FVector& HeightAxis, const float Radius, const float Thickness, const FColor& Color, const FString& Description = FString());
+	static FGameplayDebuggerShape MakeRectangle(const FVector& Center, const FVector& WidthAxis, const FVector& HeightAxis, const float Width, const float Height, const FColor& Color, const FString& Description = FString());
+	static FGameplayDebuggerShape MakeRectangle(const FVector& Center, const FVector& WidthAxis, const FVector& HeightAxis, const float Width, const float Height, const float Thickness, const FColor& Color, const FString& Description = FString());
 	static FGameplayDebuggerShape MakeCapsule(const FVector& Center, const float Radius, const float HalfHeight, const FColor& Color, const FString& Description = FString());
-	static FGameplayDebuggerShape MakePolygon(const TArray<FVector>& Verts, const FColor& Color, const FString& Description = FString());
+	static FGameplayDebuggerShape MakePolygon(TConstArrayView<FVector> Verts, const FColor& Color, const FString& Description = FString());
+	static FGameplayDebuggerShape MakePolyline(const TConstArrayView<FVector> Verts, const float Thickness, const FColor& Color, const FString& Description = FString());
+	static FGameplayDebuggerShape MakePolyline(const TConstArrayView<FVector> Verts, const FColor& Color, const FString& Description = FString());
+	static FGameplayDebuggerShape MakeSegmentList(const TConstArrayView<FVector> Verts, const float Thickness, const FColor& Color, const FString& Description = FString());
+	static FGameplayDebuggerShape MakeSegmentList(const TConstArrayView<FVector> Verts, const FColor& Color, const FString& Description = FString());
 
 	void Draw(UWorld* World, FGameplayDebuggerCanvasContext& Context);
 };
