Minor fix for previous submit accidentally replaing 0.2f with 0.

diff --git a/Engine/Source/Runtime/GameplayDebugger/Private/GameplayDebuggerLocalController.cpp b/Engine/Source/Runtime/GameplayDebugger/Private/GameplayDebuggerLocalController.cpp
--- a/Engine/Source/Runtime/GameplayDebugger/Private/GameplayDebuggerLocalController.cpp
+++ b/Engine/Source/Runtime/GameplayDebugger/Private/GameplayDebuggerLocalController.cpp
@@ -555,7 +555,7 @@ void UGameplayDebuggerLocalController::OnActivationPressed()
 {
 	if (CachedReplicator)
 	{
-		const double HoldTimeThr = 0. * (FApp::UseFixedTimeStep() ? (FApp::GetFixedDeltaTime() * 60.) : 1.);
+		const double HoldTimeThr = 0.2 * (FApp::UseFixedTimeStep() ? (FApp::GetFixedDeltaTime() * 60.) : 1.);
 
 		CachedReplicator->GetWorldTimerManager().SetTimer(StartSelectingActorHandle, this, &UGameplayDebuggerLocalController::OnStartSelectingActor, static_cast<float>(HoldTimeThr));
 	}
