Gameplay Debugger: Added oriented box and capsule.

diff --git a/Engine/Source/Runtime/GameplayDebugger/Private/GameplayDebuggerTypes.cpp b/Engine/Source/Runtime/GameplayDebugger/Private/GameplayDebuggerTypes.cpp
--- a/Engine/Source/Runtime/GameplayDebugger/Private/GameplayDebuggerTypes.cpp
+++ b/Engine/Source/Runtime/GameplayDebugger/Private/GameplayDebuggerTypes.cpp
@@ -71,10 +71,35 @@ FGameplayDebuggerShape FGameplayDebuggerShape::MakeArrow(const FVector& StartLoc
 }
 
 FGameplayDebuggerShape FGameplayDebuggerShape::MakeBox(const FVector& Center, const FVector& Extent, const FColor& Color, const FString& Description)
+{
+	return MakeBox(Center, Extent, 1.0f, Color, Description);
+}
+
+FGameplayDebuggerShape FGameplayDebuggerShape::MakeBox(const FVector& Center, const FVector& Extent, const float Thickness, const FColor& Color, const FString& Description)
 {
 	FGameplayDebuggerShape NewElement;
 	NewElement.ShapeData.Add(Center);
 	NewElement.ShapeData.Add(Extent);
+	NewElement.ShapeData.Add(FVector(Thickness, 0, 0));
+	NewElement.Color = Color;
+	NewElement.Description = Description;
+	NewElement.Type = EGameplayDebuggerShape::Box;
+
+	return NewElement;
+}
+
+FGameplayDebuggerShape FGameplayDebuggerShape::MakeBox(const FVector& Center, const FRotator& Rotation, const FVector& Extent, const FColor& Color, const FString& Description)
+{
+	return MakeBox(Center, Rotation, Extent, 1.0f, Color, Description);
+}
+
+FGameplayDebuggerShape FGameplayDebuggerShape::MakeBox(const FVector& Center, const FRotator& Rotation, const FVector& Extent, const float Thickness, const FColor& Color, const FString& Description)
+{
+	FGameplayDebuggerShape NewElement;
+	NewElement.ShapeData.Add(Center);
+	NewElement.ShapeData.Add(Extent);
+	NewElement.ShapeData.Add(Rotation.Euler());
+	NewElement.ShapeData.Add(FVector(Thickness, 0, 0));
 	NewElement.Color = Color;
 	NewElement.Description = Description;
 	NewElement.Type = EGameplayDebuggerShape::Box;
@@ -167,10 +192,16 @@ FGameplayDebuggerShape FGameplayDebuggerShape::MakeRectangle(const FVector& Cent
 }
 
 FGameplayDebuggerShape FGameplayDebuggerShape::MakeCapsule(const FVector& Center, const float Radius, const float HalfHeight, const FColor& Color, const FString& Description)
+{
+	return MakeCapsule(Center, FRotator::ZeroRotator, Radius, HalfHeight, Color, Description);
+}
+
+FGameplayDebuggerShape FGameplayDebuggerShape::MakeCapsule(const FVector& Center, const FRotator& Rotation, const float Radius, const float HalfHeight, const FColor& Color, const FString& Description)
 {
 	FGameplayDebuggerShape NewElement;
 	NewElement.ShapeData.Add(Center);
 	NewElement.ShapeData.Add(FVector(Radius, 0, HalfHeight));
+	NewElement.ShapeData.Add(Rotation.Euler());
 	NewElement.Color = Color;
 	NewElement.Description = Description;
 	NewElement.Type = EGameplayDebuggerShape::Capsule;
@@ -264,9 +295,16 @@ void FGameplayDebuggerShape::Draw(UWorld* World, FGameplayDebuggerCanvasContext&
 		break;
 
 	case EGameplayDebuggerShape::Box:
-		if (ShapeData.Num() == 2)
+		if (ShapeData.Num() == 3)
+		{
+			const float Thickness = ShapeData[2].X;
+			DrawDebugBox(World, ShapeData[0], ShapeData[1], Color, bPersistent, LifeTime, DepthPriority, Thickness);
+			DescLocation = ShapeData[0];
+		}
+		else if (ShapeData.Num() == 4)
 		{
-			DrawDebugBox(World, ShapeData[0], ShapeData[1], Color);
+			const float Thickness = ShapeData[3].X;
+			DrawDebugBox(World, ShapeData[0], ShapeData[1], FQuat::MakeFromEuler(ShapeData[2]), Color, bPersistent, LifeTime, DepthPriority, Thickness);
 			DescLocation = ShapeData[0];
 		}
 		break;
@@ -297,9 +335,9 @@ void FGameplayDebuggerShape::Draw(UWorld* World, FGameplayDebuggerCanvasContext&
 		break;
 
 	case EGameplayDebuggerShape::Capsule:
-		if (ShapeData.Num() == 2)
+		if (ShapeData.Num() == 3)
 		{
-			DrawDebugCapsule(World, ShapeData[0], static_cast<float>(ShapeData[1].Z), static_cast<float>(ShapeData[1].X), FQuat::Identity, Color, bPersistent, LifeTime, DepthPriority);
+			DrawDebugCapsule(World, ShapeData[0], static_cast<float>(ShapeData[1].Z), static_cast<float>(ShapeData[1].X), FQuat::MakeFromEuler(ShapeData[2]), Color, bPersistent, LifeTime, DepthPriority);
 			DescLocation = ShapeData[0];
 		}
 		break;
diff --git a/Engine/Source/Runtime/GameplayDebugger/Public/GameplayDebuggerTypes.h b/Engine/Source/Runtime/GameplayDebugger/Public/GameplayDebuggerTypes.h
--- a/Engine/Source/Runtime/GameplayDebugger/Public/GameplayDebuggerTypes.h
+++ b/Engine/Source/Runtime/GameplayDebugger/Public/GameplayDebuggerTypes.h
@@ -209,6 +209,9 @@ struct GAMEPLAYDEBUGGER_API FGameplayDebuggerShape
 	static FGameplayDebuggerShape MakeSegment(const FVector& StartLocation, const FVector& EndLocation, const FColor& Color, const FString& Description = FString());
 	static FGameplayDebuggerShape MakeArrow(const FVector& StartLocation, const FVector& EndLocation, const float HeadSize, const float Thickness, const FColor& Color, const FString& Description = FString());
 	static FGameplayDebuggerShape MakeBox(const FVector& Center, const FVector& Extent, const FColor& Color, const FString& Description = FString());
+	static FGameplayDebuggerShape MakeBox(const FVector& Center, const FVector& Extent, const float Thickness, const FColor& Color, const FString& Description = FString());
+	static FGameplayDebuggerShape MakeBox(const FVector& Center, const FRotator& Rotation, const FVector& Extent, const FColor& Color, const FString& Description = FString());
+	static FGameplayDebuggerShape MakeBox(const FVector& Center, const FRotator& Rotation, const FVector& Extent, const float Thickness, const FColor& Color, const FString& Description = FString());
 	static FGameplayDebuggerShape MakeCone(const FVector& Location, const FVector& Direction, const float Length, const FColor& Color, const FString& Description = FString());
 	static FGameplayDebuggerShape MakeCylinder(const FVector& Center, const float Radius, const float HalfHeight, const FColor& Color, const FString& Description = FString());
 	static FGameplayDebuggerShape MakeCircle(const FVector& Center, const FVector& Up, const float Radius, const FColor& Color, const FString& Description = FString());
@@ -218,6 +221,7 @@ struct GAMEPLAYDEBUGGER_API FGameplayDebuggerShape
 	static FGameplayDebuggerShape MakeRectangle(const FVector& Center, const FVector& WidthAxis, const FVector& HeightAxis, const float Width, const float Height, const FColor& Color, const FString& Description = FString());
 	static FGameplayDebuggerShape MakeRectangle(const FVector& Center, const FVector& WidthAxis, const FVector& HeightAxis, const float Width, const float Height, const float Thickness, const FColor& Color, const FString& Description = FString());
 	static FGameplayDebuggerShape MakeCapsule(const FVector& Center, const float Radius, const float HalfHeight, const FColor& Color, const FString& Description = FString());
+	static FGameplayDebuggerShape MakeCapsule(const FVector& Center, const FRotator& Rotation, const float Radius, const float HalfHeight, const FColor& Color, const FString& Description = FString());
 	static FGameplayDebuggerShape MakePolygon(TConstArrayView<FVector> Verts, const FColor& Color, const FString& Description = FString());
 	static FGameplayDebuggerShape MakePolyline(const TConstArrayView<FVector> Verts, const float Thickness, const FColor& Color, const FString& Description = FString());
 	static FGameplayDebuggerShape MakePolyline(const TConstArrayView<FVector> Verts, const FColor& Color, const FString& Description = FString());
