Move display spawn points from LWM to GameplayDebugger

diff --git a/Engine/Source/Runtime/GameplayDebugger/Private/GameplayDebuggerLocalController.cpp b/Engine/Source/Runtime/GameplayDebugger/Private/GameplayDebuggerLocalController.cpp
--- a/Engine/Source/Runtime/GameplayDebugger/Private/GameplayDebuggerLocalController.cpp
+++ b/Engine/Source/Runtime/GameplayDebugger/Private/GameplayDebuggerLocalController.cpp
@@ -930,6 +930,23 @@ private:
 		return Controller;
 	}
 
+	static void EnableGameplayDebugger(const TArray<FString>& Args, UWorld* InWorld)
+	{
+		if (UGameplayDebuggerLocalController* Controller = GetController(InWorld))
+		{
+			bool bEnable = true;
+			if (Args.Num() > 0)
+			{
+				LexFromString(bEnable, *Args[0]);
+			}
+
+			if (Controller->bIsLocallyEnabled != bEnable)
+			{
+				Controller->ToggleActivation();	
+			}
+		}
+	}
+
 	static void ToggleGameplayDebugger(UWorld* InWorld)
 	{
 		if (UGameplayDebuggerLocalController* Controller = GetController(InWorld))
@@ -1056,9 +1073,10 @@ private:
 	}
 
 	/** For legacy command: EnableGDT */
-	static FAutoConsoleCommandWithWorld EnableDebuggerCmd;
+	static FAutoConsoleCommandWithWorld LegacyEnableDebuggerCmd;
 
 	/** Various gameplay debugger commands: gdt.<command> */
+	static FAutoConsoleCommandWithWorldAndArgs EnableDebuggerCmd;
 	static FAutoConsoleCommandWithWorld ToggleDebuggerCmd;
 	static FAutoConsoleCommandWithWorld SelectLocalPlayerCmd;
 	static FAutoConsoleCommandWithWorld SelectPreviousRowCmd;
@@ -1068,12 +1086,18 @@ private:
 	static FAutoConsoleCommandWithWorldAndArgs SetFontSizeCmd;
 };
 
-FAutoConsoleCommandWithWorld FGameplayDebuggerConsoleCommands::EnableDebuggerCmd(
+FAutoConsoleCommandWithWorld FGameplayDebuggerConsoleCommands::LegacyEnableDebuggerCmd(
 	TEXT("EnableGDT"),
 	TEXT("Toggles Gameplay Debugger Tool"),
 	FConsoleCommandWithWorldDelegate::CreateStatic(&FGameplayDebuggerConsoleCommands::ToggleGameplayDebugger)
 );
 
+FAutoConsoleCommandWithWorldAndArgs FGameplayDebuggerConsoleCommands::EnableDebuggerCmd(
+	TEXT("gdt.Enable"),
+	TEXT("Enable Gameplay Debugger Tool"),
+	FConsoleCommandWithWorldAndArgsDelegate::CreateStatic(&FGameplayDebuggerConsoleCommands::EnableGameplayDebugger)
+);
+
 FAutoConsoleCommandWithWorld FGameplayDebuggerConsoleCommands::ToggleDebuggerCmd(
 	TEXT("gdt.Toggle"),
 	TEXT("Toggles Gameplay Debugger Tool"),
