[Engine]

diff --git a/Engine/Source/Runtime/GameplayDebugger/Public/GameplayDebuggerTypes.h b/Engine/Source/Runtime/GameplayDebugger/Public/GameplayDebuggerTypes.h
--- a/Engine/Source/Runtime/GameplayDebugger/Public/GameplayDebuggerTypes.h
+++ b/Engine/Source/Runtime/GameplayDebugger/Public/GameplayDebuggerTypes.h
@@ -26,7 +26,7 @@ namespace FGameplayDebuggerUtils
 	bool GAMEPLAYDEBUGGER_API IsAuthority(UWorld* World);
 }
 
-class GAMEPLAYDEBUGGER_API FGameplayDebuggerCanvasContext
+class FGameplayDebuggerCanvasContext
 {
 public:
 	/** canvas used for drawing */
@@ -51,15 +51,15 @@ public:
 	float DefaultX, DefaultY;
 
 	FGameplayDebuggerCanvasContext() {}
-	FGameplayDebuggerCanvasContext(UCanvas* InCanvas, UFont* InFont);
+	GAMEPLAYDEBUGGER_API FGameplayDebuggerCanvasContext(UCanvas* InCanvas, UFont* InFont);
 
 	// print string on canvas
-	void Print(const FString& String);
-	void Print(const FColor& Color, const FString& String);
-	void Print(const FColor& Color, const float Alpha, const FString& String);
-	void PrintAt(float PosX, float PosY, const FString& String);
-	void PrintAt(float PosX, float PosY, const FColor& Color, const FString& String);
-	void PrintAt(float PosX, float PosY, const FColor& Color, const float Alpha, const FString& String);
+	GAMEPLAYDEBUGGER_API void Print(const FString& String);
+	GAMEPLAYDEBUGGER_API void Print(const FColor& Color, const FString& String);
+	GAMEPLAYDEBUGGER_API void Print(const FColor& Color, const float Alpha, const FString& String);
+	GAMEPLAYDEBUGGER_API void PrintAt(float PosX, float PosY, const FString& String);
+	GAMEPLAYDEBUGGER_API void PrintAt(float PosX, float PosY, const FColor& Color, const FString& String);
+	GAMEPLAYDEBUGGER_API void PrintAt(float PosX, float PosY, const FColor& Color, const float Alpha, const FString& String);
 
 	// print formatted string on canvas
 	template <typename FmtType, typename... Types>
@@ -117,34 +117,34 @@ public:
 	}
 
 private:
-	void VARARGS PrintfImpl(const FColor& Color, const float Alpha, const TCHAR* Args, ...);
-	void VARARGS PrintfAtImpl(float PosX, float PosY, const FColor& Color, const float Alpha, const TCHAR* Args, ...);
+	GAMEPLAYDEBUGGER_API void VARARGS PrintfImpl(const FColor& Color, const float Alpha, const TCHAR* Args, ...);
+	GAMEPLAYDEBUGGER_API void VARARGS PrintfAtImpl(float PosX, float PosY, const FColor& Color, const float Alpha, const TCHAR* Args, ...);
 
 public:
 	// moves cursor to new line
-	void MoveToNewLine();
+	GAMEPLAYDEBUGGER_API void MoveToNewLine();
 
 	// calculate size of string
-	void MeasureString(const FString& String, float& OutSizeX, float& OutSizeY) const;
+	GAMEPLAYDEBUGGER_API void MeasureString(const FString& String, float& OutSizeX, float& OutSizeY) const;
 
 	// get height of single line text
-	float GetLineHeight() const;
+	GAMEPLAYDEBUGGER_API float GetLineHeight() const;
 
 	// project world location on canvas
-	FVector2D ProjectLocation(const FVector& Location) const;
+	GAMEPLAYDEBUGGER_API FVector2D ProjectLocation(const FVector& Location) const;
 
 	// check if world location is visible in current view
-	bool IsLocationVisible(const FVector& Location) const;
+	GAMEPLAYDEBUGGER_API bool IsLocationVisible(const FVector& Location) const;
 
 	// draw item on canvas
-	void DrawItem(FCanvasItem& Item, float PosX, float PosY);
+	GAMEPLAYDEBUGGER_API void DrawItem(FCanvasItem& Item, float PosX, float PosY);
 
 	// draw icon on canvas
-	void DrawIcon(const FColor& Color, const FCanvasIcon& Icon, float PosX, float PosY, float Scale = 1.f);
+	GAMEPLAYDEBUGGER_API void DrawIcon(const FColor& Color, const FCanvasIcon& Icon, float PosX, float PosY, float Scale = 1.f);
 
 	// fetches the World associated with this context. Will use the World member variable, if set, and 
 	// PlayerController's world otherwise. Note that it can still be null and needs to be tested.
-	UWorld* GetWorld() const;
+	GAMEPLAYDEBUGGER_API UWorld* GetWorld() const;
 };
 
 namespace FGameplayDebuggerCanvasStrings
@@ -177,7 +177,7 @@ enum class EGameplayDebuggerShape : uint8
 };
 
 USTRUCT()
-struct GAMEPLAYDEBUGGER_API FGameplayDebuggerShape
+struct FGameplayDebuggerShape
 {
 	GENERATED_BODY()
 
@@ -204,31 +204,31 @@ struct GAMEPLAYDEBUGGER_API FGameplayDebuggerShape
 		return (Type == Other.Type) && (Color == Other.Color) && (Description == Other.Description) && (ShapeData == Other.ShapeData);
 	}
 
-	static FGameplayDebuggerShape MakePoint(const FVector& Location, const float Radius, const FColor& Color, const FString& Description = FString());
-	static FGameplayDebuggerShape MakeSegment(const FVector& StartLocation, const FVector& EndLocation, const float Thickness, const FColor& Color, const FString& Description = FString());
-	static FGameplayDebuggerShape MakeSegment(const FVector& StartLocation, const FVector& EndLocation, const FColor& Color, const FString& Description = FString());
-	static FGameplayDebuggerShape MakeArrow(const FVector& StartLocation, const FVector& EndLocation, const float HeadSize, const float Thickness, const FColor& Color, const FString& Description = FString());
-	static FGameplayDebuggerShape MakeBox(const FVector& Center, const FVector& Extent, const FColor& Color, const FString& Description = FString());
-	static FGameplayDebuggerShape MakeBox(const FVector& Center, const FVector& Extent, const float Thickness, const FColor& Color, const FString& Description = FString());
-	static FGameplayDebuggerShape MakeBox(const FVector& Center, const FRotator& Rotation, const FVector& Extent, const FColor& Color, const FString& Description = FString());
-	static FGameplayDebuggerShape MakeBox(const FVector& Center, const FRotator& Rotation, const FVector& Extent, const float Thickness, const FColor& Color, const FString& Description = FString());
-	static FGameplayDebuggerShape MakeCone(const FVector& Location, const FVector& Direction, const float Length, const FColor& Color, const FString& Description = FString());
-	static FGameplayDebuggerShape MakeCylinder(const FVector& Center, const float Radius, const float HalfHeight, const FColor& Color, const FString& Description = FString());
-	static FGameplayDebuggerShape MakeCircle(const FVector& Center, const FVector& Up, const float Radius, const FColor& Color, const FString& Description = FString());
-	static FGameplayDebuggerShape MakeCircle(const FVector& Center, const FVector& Up, const float Radius, const float Thickness, const FColor& Color, const FString& Description = FString());
-	static FGameplayDebuggerShape MakeCircle(const FVector& Center, const FVector& WidthAxis, const FVector& HeightAxis, const float Radius, const FColor& Color, const FString& Description = FString());
-	static FGameplayDebuggerShape MakeCircle(const FVector& Center, const FVector& WidthAxis, const FVector& HeightAxis, const float Radius, const float Thickness, const FColor& Color, const FString& Description = FString());
-	static FGameplayDebuggerShape MakeRectangle(const FVector& Center, const FVector& WidthAxis, const FVector& HeightAxis, const float Width, const float Height, const FColor& Color, const FString& Description = FString());
-	static FGameplayDebuggerShape MakeRectangle(const FVector& Center, const FVector& WidthAxis, const FVector& HeightAxis, const float Width, const float Height, const float Thickness, const FColor& Color, const FString& Description = FString());
-	static FGameplayDebuggerShape MakeCapsule(const FVector& Center, const float Radius, const float HalfHeight, const FColor& Color, const FString& Description = FString());
-	static FGameplayDebuggerShape MakeCapsule(const FVector& Center, const FRotator& Rotation, const float Radius, const float HalfHeight, const FColor& Color, const FString& Description = FString());
-	static FGameplayDebuggerShape MakePolygon(TConstArrayView<FVector> Verts, const FColor& Color, const FString& Description = FString());
-	static FGameplayDebuggerShape MakePolyline(const TConstArrayView<FVector> Verts, const float Thickness, const FColor& Color, const FString& Description = FString());
-	static FGameplayDebuggerShape MakePolyline(const TConstArrayView<FVector> Verts, const FColor& Color, const FString& Description = FString());
-	static FGameplayDebuggerShape MakeSegmentList(const TConstArrayView<FVector> Verts, const float Thickness, const FColor& Color, const FString& Description = FString());
-	static FGameplayDebuggerShape MakeSegmentList(const TConstArrayView<FVector> Verts, const FColor& Color, const FString& Description = FString());
-
-	void Draw(UWorld* World, FGameplayDebuggerCanvasContext& Context);
+	static GAMEPLAYDEBUGGER_API FGameplayDebuggerShape MakePoint(const FVector& Location, const float Radius, const FColor& Color, const FString& Description = FString());
+	static GAMEPLAYDEBUGGER_API FGameplayDebuggerShape MakeSegment(const FVector& StartLocation, const FVector& EndLocation, const float Thickness, const FColor& Color, const FString& Description = FString());
+	static GAMEPLAYDEBUGGER_API FGameplayDebuggerShape MakeSegment(const FVector& StartLocation, const FVector& EndLocation, const FColor& Color, const FString& Description = FString());
+	static GAMEPLAYDEBUGGER_API FGameplayDebuggerShape MakeArrow(const FVector& StartLocation, const FVector& EndLocation, const float HeadSize, const float Thickness, const FColor& Color, const FString& Description = FString());
+	static GAMEPLAYDEBUGGER_API FGameplayDebuggerShape MakeBox(const FVector& Center, const FVector& Extent, const FColor& Color, const FString& Description = FString());
+	static GAMEPLAYDEBUGGER_API FGameplayDebuggerShape MakeBox(const FVector& Center, const FVector& Extent, const float Thickness, const FColor& Color, const FString& Description = FString());
+	static GAMEPLAYDEBUGGER_API FGameplayDebuggerShape MakeBox(const FVector& Center, const FRotator& Rotation, const FVector& Extent, const FColor& Color, const FString& Description = FString());
+	static GAMEPLAYDEBUGGER_API FGameplayDebuggerShape MakeBox(const FVector& Center, const FRotator& Rotation, const FVector& Extent, const float Thickness, const FColor& Color, const FString& Description = FString());
+	static GAMEPLAYDEBUGGER_API FGameplayDebuggerShape MakeCone(const FVector& Location, const FVector& Direction, const float Length, const FColor& Color, const FString& Description = FString());
+	static GAMEPLAYDEBUGGER_API FGameplayDebuggerShape MakeCylinder(const FVector& Center, const float Radius, const float HalfHeight, const FColor& Color, const FString& Description = FString());
+	static GAMEPLAYDEBUGGER_API FGameplayDebuggerShape MakeCircle(const FVector& Center, const FVector& Up, const float Radius, const FColor& Color, const FString& Description = FString());
+	static GAMEPLAYDEBUGGER_API FGameplayDebuggerShape MakeCircle(const FVector& Center, const FVector& Up, const float Radius, const float Thickness, const FColor& Color, const FString& Description = FString());
+	static GAMEPLAYDEBUGGER_API FGameplayDebuggerShape MakeCircle(const FVector& Center, const FVector& WidthAxis, const FVector& HeightAxis, const float Radius, const FColor& Color, const FString& Description = FString());
+	static GAMEPLAYDEBUGGER_API FGameplayDebuggerShape MakeCircle(const FVector& Center, const FVector& WidthAxis, const FVector& HeightAxis, const float Radius, const float Thickness, const FColor& Color, const FString& Description = FString());
+	static GAMEPLAYDEBUGGER_API FGameplayDebuggerShape MakeRectangle(const FVector& Center, const FVector& WidthAxis, const FVector& HeightAxis, const float Width, const float Height, const FColor& Color, const FString& Description = FString());
+	static GAMEPLAYDEBUGGER_API FGameplayDebuggerShape MakeRectangle(const FVector& Center, const FVector& WidthAxis, const FVector& HeightAxis, const float Width, const float Height, const float Thickness, const FColor& Color, const FString& Description = FString());
+	static GAMEPLAYDEBUGGER_API FGameplayDebuggerShape MakeCapsule(const FVector& Center, const float Radius, const float HalfHeight, const FColor& Color, const FString& Description = FString());
+	static GAMEPLAYDEBUGGER_API FGameplayDebuggerShape MakeCapsule(const FVector& Center, const FRotator& Rotation, const float Radius, const float HalfHeight, const FColor& Color, const FString& Description = FString());
+	static GAMEPLAYDEBUGGER_API FGameplayDebuggerShape MakePolygon(TConstArrayView<FVector> Verts, const FColor& Color, const FString& Description = FString());
+	static GAMEPLAYDEBUGGER_API FGameplayDebuggerShape MakePolyline(const TConstArrayView<FVector> Verts, const float Thickness, const FColor& Color, const FString& Description = FString());
+	static GAMEPLAYDEBUGGER_API FGameplayDebuggerShape MakePolyline(const TConstArrayView<FVector> Verts, const FColor& Color, const FString& Description = FString());
+	static GAMEPLAYDEBUGGER_API FGameplayDebuggerShape MakeSegmentList(const TConstArrayView<FVector> Verts, const float Thickness, const FColor& Color, const FString& Description = FString());
+	static GAMEPLAYDEBUGGER_API FGameplayDebuggerShape MakeSegmentList(const TConstArrayView<FVector> Verts, const FColor& Color, const FString& Description = FString());
+
+	GAMEPLAYDEBUGGER_API void Draw(UWorld* World, FGameplayDebuggerCanvasContext& Context);
 };
 
 FArchive& operator<<(FArchive& Ar, FGameplayDebuggerShape& Shape);
@@ -283,7 +283,7 @@ struct FGameplayDebuggerDataPackHeader
 	}
 };
 
-struct GAMEPLAYDEBUGGER_API FGameplayDebuggerDataPack
+struct FGameplayDebuggerDataPack
 {
 	using FHeader = FGameplayDebuggerDataPackHeader;
 
@@ -322,13 +322,13 @@ struct GAMEPLAYDEBUGGER_API FGameplayDebuggerDataPack
 	FOnReset ResetDelegate;
 	FOnSerialize SerializeDelegate;
 
-	static int32 PacketSize;
+	static GAMEPLAYDEBUGGER_API int32 PacketSize;
 
-	bool CheckDirtyAndUpdate();
-	bool RequestReplication(int16 SyncCounter);
+	GAMEPLAYDEBUGGER_API bool CheckDirtyAndUpdate();
+	GAMEPLAYDEBUGGER_API bool RequestReplication(int16 SyncCounter);
 
-	void OnReplicated();
-	void OnPacketRequest(int16 DataVersion, int32 DataOffset);
+	GAMEPLAYDEBUGGER_API void OnReplicated();
+	GAMEPLAYDEBUGGER_API void OnPacketRequest(int16 DataVersion, int32 DataOffset);
 
 	/** get replication progress in (0..1) range */
 	FORCEINLINE float GetProgress() const
@@ -357,7 +357,7 @@ enum class EGameplayDebuggerInputMode : uint8
 	Replicated,
 };
 
-struct GAMEPLAYDEBUGGER_API FGameplayDebuggerInputModifier
+struct FGameplayDebuggerInputModifier
 {
 	uint32 bShift : 1;
 	uint32 bCtrl : 1;
@@ -374,14 +374,14 @@ struct GAMEPLAYDEBUGGER_API FGameplayDebuggerInputModifier
 		return FGameplayDebuggerInputModifier(bShift || Other.bShift, bCtrl || Other.bCtrl, bAlt || Other.bAlt, bCmd || Other.bCmd);
 	}
 
-	static FGameplayDebuggerInputModifier Shift;
-	static FGameplayDebuggerInputModifier Ctrl;
-	static FGameplayDebuggerInputModifier Alt;
-	static FGameplayDebuggerInputModifier Cmd;
-	static FGameplayDebuggerInputModifier None;
+	static GAMEPLAYDEBUGGER_API FGameplayDebuggerInputModifier Shift;
+	static GAMEPLAYDEBUGGER_API FGameplayDebuggerInputModifier Ctrl;
+	static GAMEPLAYDEBUGGER_API FGameplayDebuggerInputModifier Alt;
+	static GAMEPLAYDEBUGGER_API FGameplayDebuggerInputModifier Cmd;
+	static GAMEPLAYDEBUGGER_API FGameplayDebuggerInputModifier None;
 };
 
-struct GAMEPLAYDEBUGGER_API FGameplayDebuggerInputHandler
+struct FGameplayDebuggerInputHandler
 {
 	DECLARE_DELEGATE(FHandler);
 
@@ -392,8 +392,8 @@ struct GAMEPLAYDEBUGGER_API FGameplayDebuggerInputHandler
 
 	FGameplayDebuggerInputHandler() : KeyName(NAME_None), Mode(EGameplayDebuggerInputMode::Local) {}
 
-	bool IsValid() const;
-	FString ToString() const;
+	GAMEPLAYDEBUGGER_API bool IsValid() const;
+	GAMEPLAYDEBUGGER_API FString ToString() const;
 };
 
 /** 
@@ -401,18 +401,18 @@ struct GAMEPLAYDEBUGGER_API FGameplayDebuggerInputHandler
  * Intended to use only from addon's constructor!
  * Check example in FGameplayDebuggerExtension_Spectator::FGameplayDebuggerExtension_Spectator() for details
  */
-struct GAMEPLAYDEBUGGER_API FGameplayDebuggerInputHandlerConfig
+struct FGameplayDebuggerInputHandlerConfig
 {
 	FName KeyName;
 	FGameplayDebuggerInputModifier Modifier;
 
 	FGameplayDebuggerInputHandlerConfig() : KeyName(NAME_None) {}
-	FGameplayDebuggerInputHandlerConfig(const FName ConfigName, const FName DefaultKeyName);
-	FGameplayDebuggerInputHandlerConfig(const FName ConfigName, const FName DefaultKeyName, const FGameplayDebuggerInputModifier& DefaultModifier);
+	GAMEPLAYDEBUGGER_API FGameplayDebuggerInputHandlerConfig(const FName ConfigName, const FName DefaultKeyName);
+	GAMEPLAYDEBUGGER_API FGameplayDebuggerInputHandlerConfig(const FName ConfigName, const FName DefaultKeyName, const FGameplayDebuggerInputModifier& DefaultModifier);
 
-	static FName CurrentCategoryName;
-	static FName CurrentExtensionName;
+	static GAMEPLAYDEBUGGER_API FName CurrentCategoryName;
+	static GAMEPLAYDEBUGGER_API FName CurrentExtensionName;
 
 private:
-	void UpdateConfig(const FName ConfigName);
+	GAMEPLAYDEBUGGER_API void UpdateConfig(const FName ConfigName);
 };
diff --git a/Engine/Source/Runtime/GameplayMediaEncoder/Public/GameplayMediaEncoder.h b/Engine/Source/Runtime/GameplayMediaEncoder/Public/GameplayMediaEncoder.h
--- a/Engine/Source/Runtime/GameplayMediaEncoder/Public/GameplayMediaEncoder.h
+++ b/Engine/Source/Runtime/GameplayMediaEncoder/Public/GameplayMediaEncoder.h
@@ -33,22 +33,22 @@ public:
 	virtual void OnMediaSample(const AVEncoder::FMediaPacket& Sample) = 0;
 };
 
-class GAMEPLAYMEDIAENCODER_API FGameplayMediaEncoder final : private ISubmixBufferListener, public AVEncoder::IAudioEncoderListener
+class FGameplayMediaEncoder final : private ISubmixBufferListener, public AVEncoder::IAudioEncoderListener
 {
 public:
 
 	/**
 	 * Get the singleton
 	 */
-	static FGameplayMediaEncoder* Get();
+	static GAMEPLAYMEDIAENCODER_API FGameplayMediaEncoder* Get();
 
-	~FGameplayMediaEncoder();
+	GAMEPLAYMEDIAENCODER_API ~FGameplayMediaEncoder();
 
-	bool RegisterListener(IGameplayMediaEncoderListener* Listener);
-	void UnregisterListener(IGameplayMediaEncoderListener* Listener);
+	GAMEPLAYMEDIAENCODER_API bool RegisterListener(IGameplayMediaEncoderListener* Listener);
+	GAMEPLAYMEDIAENCODER_API void UnregisterListener(IGameplayMediaEncoderListener* Listener);
 
-	void SetVideoBitrate(uint32 Bitrate);
-	void SetVideoFramerate(uint32 Framerate);
+	GAMEPLAYMEDIAENCODER_API void SetVideoBitrate(uint32 Bitrate);
+	GAMEPLAYMEDIAENCODER_API void SetVideoFramerate(uint32 Framerate);
 
 	///**
 	// * Returns the audio codec name and configuration
@@ -56,10 +56,10 @@ public:
 	//TPair<FString, AVEncoder::FAudioConfig> GetAudioConfig() const;
 	//TPair<FString, AVEncoder::FVideoConfig> GetVideoConfig() const;
 
-	bool Initialize();
-	void Shutdown();
-	bool Start();
-	void Stop();
+	GAMEPLAYMEDIAENCODER_API bool Initialize();
+	GAMEPLAYMEDIAENCODER_API void Shutdown();
+	GAMEPLAYMEDIAENCODER_API bool Start();
+	GAMEPLAYMEDIAENCODER_API void Stop();
 
 	static void InitializeCmd()
 	{
@@ -85,34 +85,34 @@ public:
 		Get()->Stop();
 	}
 
-	AVEncoder::FAudioConfig GetAudioConfig() const;
+	GAMEPLAYMEDIAENCODER_API AVEncoder::FAudioConfig GetAudioConfig() const;
 	AVEncoder::FVideoConfig GetVideoConfig() const { return VideoConfig; }
 
 private:
 
 	// Private to control how our single instance is created
-	FGameplayMediaEncoder();
+	GAMEPLAYMEDIAENCODER_API FGameplayMediaEncoder();
 
 	// Returns how long it has been recording for.
-	FTimespan GetMediaTimestamp() const;
+	GAMEPLAYMEDIAENCODER_API FTimespan GetMediaTimestamp() const;
 
 	// Back buffer capture
-	void OnFrameBufferReady(SWindow& SlateWindow, const FTexture2DRHIRef& FrameBuffer);
+	GAMEPLAYMEDIAENCODER_API void OnFrameBufferReady(SWindow& SlateWindow, const FTexture2DRHIRef& FrameBuffer);
 	// ISubmixBufferListener interface
-	void OnNewSubmixBuffer(const USoundSubmix* OwningSubmix, float* AudioData, int32 NumSamples, int32 NumChannels, const int32 SampleRate, double AudioClock) override;
+	GAMEPLAYMEDIAENCODER_API void OnNewSubmixBuffer(const USoundSubmix* OwningSubmix, float* AudioData, int32 NumSamples, int32 NumChannels, const int32 SampleRate, double AudioClock) override;
 
-	void ProcessAudioFrame(const float* AudioData, int32 NumSamples, int32 NumChannels, int32 SampleRate);
-	void ProcessVideoFrame(const FTexture2DRHIRef& FrameBuffer);
+	GAMEPLAYMEDIAENCODER_API void ProcessAudioFrame(const float* AudioData, int32 NumSamples, int32 NumChannels, int32 SampleRate);
+	GAMEPLAYMEDIAENCODER_API void ProcessVideoFrame(const FTexture2DRHIRef& FrameBuffer);
 
-	void UpdateVideoConfig();
+	GAMEPLAYMEDIAENCODER_API void UpdateVideoConfig();
 
-	void OnEncodedAudioFrame(const AVEncoder::FMediaPacket& Packet) override;
-	void OnEncodedVideoFrame(uint32 LayerIndex, const TSharedPtr<AVEncoder::FVideoEncoderInputFrame> Frame, const AVEncoder::FCodecPacket& Packet);
+	GAMEPLAYMEDIAENCODER_API void OnEncodedAudioFrame(const AVEncoder::FMediaPacket& Packet) override;
+	GAMEPLAYMEDIAENCODER_API void OnEncodedVideoFrame(uint32 LayerIndex, const TSharedPtr<AVEncoder::FVideoEncoderInputFrame> Frame, const AVEncoder::FCodecPacket& Packet);
 
-	TSharedPtr<AVEncoder::FVideoEncoderInputFrame> ObtainInputFrame();
-	void CopyTexture(const FTexture2DRHIRef& SourceTexture, FTexture2DRHIRef& DestinationTexture) const;
+	GAMEPLAYMEDIAENCODER_API TSharedPtr<AVEncoder::FVideoEncoderInputFrame> ObtainInputFrame();
+	GAMEPLAYMEDIAENCODER_API void CopyTexture(const FTexture2DRHIRef& SourceTexture, FTexture2DRHIRef& DestinationTexture) const;
 
-	void FloatToPCM16(float const* floatSamples, int32 numSamples, TArray<int16>& out) const;
+	GAMEPLAYMEDIAENCODER_API void FloatToPCM16(float const* floatSamples, int32 numSamples, TArray<int16>& out) const;
 
 	FCriticalSection ListenersCS;
 	TArray<IGameplayMediaEncoderListener*> Listeners;
@@ -141,7 +141,7 @@ private:
 	bool bDoFrameSkipping = false;
 
 	friend class FGameplayMediaEncoderModule;
-	static FGameplayMediaEncoder* Singleton;
+	static GAMEPLAYMEDIAENCODER_API FGameplayMediaEncoder* Singleton;
 
 	// live streaming: quality adaptation to available uplink b/w
 	TAtomic<uint32> NewVideoBitrate{ 0 };
diff --git a/Engine/Source/Runtime/GameplayTasks/Classes/GameplayTask.h b/Engine/Source/Runtime/GameplayTasks/Classes/GameplayTask.h
--- a/Engine/Source/Runtime/GameplayTasks/Classes/GameplayTask.h
+++ b/Engine/Source/Runtime/GameplayTasks/Classes/GameplayTask.h
@@ -42,7 +42,7 @@ enum class ETaskResourceOverlapPolicy : uint8
 };
 	
 USTRUCT(BlueprintType)
-struct GAMEPLAYTASKS_API FGameplayResourceSet
+struct FGameplayResourceSet
 {
 	GENERATED_USTRUCT_BODY()
 
@@ -138,11 +138,11 @@ public:
 		return FGameplayResourceSet(FFlagContainer(0));
 	}
 
-	FString GetDebugDescription() const;
+	GAMEPLAYTASKS_API FString GetDebugDescription() const;
 };
 
-UCLASS(Abstract, BlueprintType, meta = (ExposedAsyncProxy=AsyncTask), config = Game)
-class GAMEPLAYTASKS_API UGameplayTask : public UObject, public IGameplayTaskOwnerInterface
+UCLASS(Abstract, BlueprintType, meta = (ExposedAsyncProxy=AsyncTask), config = Game, MinimalAPI)
+class UGameplayTask : public UObject, public IGameplayTaskOwnerInterface
 {
 	GENERATED_BODY()
 
@@ -150,46 +150,46 @@ public:
 
 	DECLARE_DYNAMIC_MULTICAST_DELEGATE(FGenericGameplayTaskDelegate);
 
-	UGameplayTask(const FObjectInitializer& ObjectInitializer);
+	GAMEPLAYTASKS_API UGameplayTask(const FObjectInitializer& ObjectInitializer);
 	
 	/** Called to trigger the actual task once the delegates have been set up */
 	UFUNCTION(BlueprintCallable, meta = (BlueprintInternalUseOnly = "true"), Category = "Gameplay Tasks")
-	void ReadyForActivation();
+	GAMEPLAYTASKS_API void ReadyForActivation();
 
 protected:
 	/** Called to trigger the actual task once the delegates have been set up
 	 *	Note that the default implementation does nothing and you don't have to call it */
-	virtual void Activate();
+	GAMEPLAYTASKS_API virtual void Activate();
 
 	/** Initializes the task with the task owner interface instance but does not activate until Activate() is called */
-	void InitTask(IGameplayTaskOwnerInterface& InTaskOwner, uint8 InPriority);
+	GAMEPLAYTASKS_API void InitTask(IGameplayTaskOwnerInterface& InTaskOwner, uint8 InPriority);
 
 public:
-	virtual void InitSimulatedTask(UGameplayTasksComponent& InGameplayTasksComponent);
+	GAMEPLAYTASKS_API virtual void InitSimulatedTask(UGameplayTasksComponent& InGameplayTasksComponent);
 
 	/** Tick function for this task, if bTickingTask == true */
 	virtual void TickTask(float DeltaTime) {}
 
 	/** Called when the task is asked to confirm from an outside node. What this means depends on the individual task. By default, this does nothing other than ending if bEndTask is true. */
-	virtual void ExternalConfirm(bool bEndTask);
+	GAMEPLAYTASKS_API virtual void ExternalConfirm(bool bEndTask);
 
 	/** Called when the task is asked to cancel from an outside node. What this means depends on the individual task. By default, this does nothing other than ending the task. */
-	virtual void ExternalCancel();
+	GAMEPLAYTASKS_API virtual void ExternalCancel();
 
 	/** Called when a task owner wants to inform the task that it's done and don't want to get notified on task deactivation (i.e. OnGameplayTaskDeactivated). */
 	void MarkOwnerFinished() { bOwnerFinished = true; }
 
 	/** Return debug string describing task */
-	virtual FString GetDebugString() const;
+	GAMEPLAYTASKS_API virtual FString GetDebugString() const;
 
 	/** Helper function for getting UWorld off a task */
-	virtual UWorld* GetWorld() const override;
+	GAMEPLAYTASKS_API virtual UWorld* GetWorld() const override;
 
 	/** Proper way to get the owning actor of task owner. This can be the owner itself since the owner is given as a interface */
-	AActor* GetOwnerActor() const;
+	GAMEPLAYTASKS_API AActor* GetOwnerActor() const;
 
 	/** Proper way to get the avatar actor associated with the task owner (usually a pawn, tower, etc) */
-	AActor* GetAvatarActor() const;
+	GAMEPLAYTASKS_API AActor* GetAvatarActor() const;
 
 	/** Helper function for instantiating and initializing a new task */
 	template <class T>
@@ -219,14 +219,14 @@ public:
 	}
 
 	/** Called when task owner has "ended" (before the task ends) kills the task. Calls OnDestroy. */
-	void TaskOwnerEnded();
+	GAMEPLAYTASKS_API void TaskOwnerEnded();
 
 	/** Called explicitly to end the task (usually by the task itself). Calls OnDestroy. 
 	 *	@NOTE: you need to call EndTask before sending out any "on completed" delegates. 
 	 *	If you don't the task will still be in an "active" state while the event receivers may
 	 *	assume it's already "finished" */
 	UFUNCTION(BlueprintCallable, Category="GameplayTasks")
-	void EndTask();
+	GAMEPLAYTASKS_API void EndTask();
 	
 	virtual bool IsSupportedForNetworking() const override { return bSimulatedTask; }
 
@@ -270,15 +270,15 @@ public:
 	 *	@note: Calling this function makes sense only until the task is being passed over to the GameplayTasksComponent.
 	 *	Once that's that resources data is consumed and further changes won't get applied 
 	 */
-	void AddRequiredResource(TSubclassOf<UGameplayTaskResource> RequiredResource);
-	void AddRequiredResourceSet(const TArray<TSubclassOf<UGameplayTaskResource> >& RequiredResourceSet);
-	void AddRequiredResourceSet(FGameplayResourceSet RequiredResourceSet);
+	GAMEPLAYTASKS_API void AddRequiredResource(TSubclassOf<UGameplayTaskResource> RequiredResource);
+	GAMEPLAYTASKS_API void AddRequiredResourceSet(const TArray<TSubclassOf<UGameplayTaskResource> >& RequiredResourceSet);
+	GAMEPLAYTASKS_API void AddRequiredResourceSet(FGameplayResourceSet RequiredResourceSet);
 
 	/** 
 	 */
-	void AddClaimedResource(TSubclassOf<UGameplayTaskResource> ClaimedResource);
-	void AddClaimedResourceSet(const TArray<TSubclassOf<UGameplayTaskResource> >& AdditionalResourcesToClaim);
-	void AddClaimedResourceSet(FGameplayResourceSet AdditionalResourcesToClaim);
+	GAMEPLAYTASKS_API void AddClaimedResource(TSubclassOf<UGameplayTaskResource> ClaimedResource);
+	GAMEPLAYTASKS_API void AddClaimedResourceSet(const TArray<TSubclassOf<UGameplayTaskResource> >& AdditionalResourcesToClaim);
+	GAMEPLAYTASKS_API void AddClaimedResourceSet(FGameplayResourceSet AdditionalResourcesToClaim);
 
 	ETaskResourceOverlapPolicy GetResourceOverlapPolicy() const { return ResourceOverlapPolicy; }
 
@@ -292,36 +292,36 @@ protected:
 	 *	IMPORTANT! When overriding this function make sure to call Super::OnDestroy(bOwnerFinished) as the last thing,
 	 *		since the function internally marks the task as "Pending Kill", and this may interfere with internal BP mechanics
 	 */
-	virtual void OnDestroy(bool bInOwnerFinished);
+	GAMEPLAYTASKS_API virtual void OnDestroy(bool bInOwnerFinished);
 
-	static IGameplayTaskOwnerInterface* ConvertToTaskOwner(UObject& OwnerObject);
-	static IGameplayTaskOwnerInterface* ConvertToTaskOwner(AActor& OwnerActor);
+	static GAMEPLAYTASKS_API IGameplayTaskOwnerInterface* ConvertToTaskOwner(UObject& OwnerObject);
+	static GAMEPLAYTASKS_API IGameplayTaskOwnerInterface* ConvertToTaskOwner(AActor& OwnerActor);
 
 	// protected by design. Not meant to be called outside from GameplayTaskComponent mechanics
-	virtual void Pause();
-	virtual void Resume();
+	GAMEPLAYTASKS_API virtual void Pause();
+	GAMEPLAYTASKS_API virtual void Resume();
 
 	// IGameplayTaskOwnerInterface BEGIN
-	virtual UGameplayTasksComponent* GetGameplayTasksComponent(const UGameplayTask& Task) const override;
-	virtual AActor* GetGameplayTaskOwner(const UGameplayTask* Task) const override;
-	virtual AActor* GetGameplayTaskAvatar(const UGameplayTask* Task) const override;
-	virtual uint8 GetGameplayTaskDefaultPriority() const override;
-	virtual void OnGameplayTaskInitialized(UGameplayTask& Task) override;
-	virtual void OnGameplayTaskDeactivated(UGameplayTask& Task) override;
+	GAMEPLAYTASKS_API virtual UGameplayTasksComponent* GetGameplayTasksComponent(const UGameplayTask& Task) const override;
+	GAMEPLAYTASKS_API virtual AActor* GetGameplayTaskOwner(const UGameplayTask* Task) const override;
+	GAMEPLAYTASKS_API virtual AActor* GetGameplayTaskAvatar(const UGameplayTask* Task) const override;
+	GAMEPLAYTASKS_API virtual uint8 GetGameplayTaskDefaultPriority() const override;
+	GAMEPLAYTASKS_API virtual void OnGameplayTaskInitialized(UGameplayTask& Task) override;
+	GAMEPLAYTASKS_API virtual void OnGameplayTaskDeactivated(UGameplayTask& Task) override;
 	// IGameplayTaskOwnerInterface END
 
 #if UE_WITH_IRIS
 	/** Register all replication fragments */
-	virtual void RegisterReplicationFragments(UE::Net::FFragmentRegistrationContext& Context, UE::Net::EFragmentRegistrationFlags RegistrationFlags) override;
+	GAMEPLAYTASKS_API virtual void RegisterReplicationFragments(UE::Net::FFragmentRegistrationContext& Context, UE::Net::EFragmentRegistrationFlags RegistrationFlags) override;
 #endif // UE_WITH_IRIS
 
 private:
 	friend UGameplayTasksComponent;
 
-	void ActivateInTaskQueue();
-	void PauseInTaskQueue();
+	GAMEPLAYTASKS_API void ActivateInTaskQueue();
+	GAMEPLAYTASKS_API void PauseInTaskQueue();
 	
-	void PerformActivation();
+	GAMEPLAYTASKS_API void PerformActivation();
 	
 protected:
 
@@ -388,8 +388,8 @@ public:
 		}
 		return DebugDescription;
 	}
-	virtual FString GenerateDebugDescription() const;
-	FString GetTaskStateName() const;
+	GAMEPLAYTASKS_API virtual FString GenerateDebugDescription() const;
+	GAMEPLAYTASKS_API FString GetTaskStateName() const;
 #endif // WITH_GAMEPLAYTASK_DEBUG
 };
 
diff --git a/Engine/Source/Runtime/GameplayTasks/Classes/GameplayTaskOwnerInterface.h b/Engine/Source/Runtime/GameplayTasks/Classes/GameplayTaskOwnerInterface.h
--- a/Engine/Source/Runtime/GameplayTasks/Classes/GameplayTaskOwnerInterface.h
+++ b/Engine/Source/Runtime/GameplayTasks/Classes/GameplayTaskOwnerInterface.h
@@ -17,16 +17,16 @@ class UGameplayTaskOwnerInterface : public UInterface
 	GENERATED_BODY()
 };
 
-class GAMEPLAYTASKS_API IGameplayTaskOwnerInterface
+class IGameplayTaskOwnerInterface
 {
 	GENERATED_BODY()
 public:
 
 	/** Finds tasks component for given GameplayTask, Task.GetGameplayTasksComponent() may not be initialized at this point! */
-	virtual UGameplayTasksComponent* GetGameplayTasksComponent(const UGameplayTask& Task) const PURE_VIRTUAL(IGameplayTaskOwnerInterface::GetGameplayTasksComponent, return nullptr;);
+	GAMEPLAYTASKS_API virtual UGameplayTasksComponent* GetGameplayTasksComponent(const UGameplayTask& Task) const PURE_VIRTUAL(IGameplayTaskOwnerInterface::GetGameplayTasksComponent, return nullptr;);
 
 	/** Get owner of a task or default one when task is null */
-	virtual AActor* GetGameplayTaskOwner(const UGameplayTask* Task) const PURE_VIRTUAL(IGameplayTaskOwnerInterface::GetGameplayTaskOwner, return nullptr;);
+	GAMEPLAYTASKS_API virtual AActor* GetGameplayTaskOwner(const UGameplayTask* Task) const PURE_VIRTUAL(IGameplayTaskOwnerInterface::GetGameplayTaskOwner, return nullptr;);
 
 	/** Get "body" of task's owner / default, having location in world (e.g. Owner = AIController, Avatar = Pawn) */
 	virtual AActor* GetGameplayTaskAvatar(const UGameplayTask* Task) const { return GetGameplayTaskOwner(Task); }
diff --git a/Engine/Source/Runtime/GameplayTasks/Classes/GameplayTaskResource.h b/Engine/Source/Runtime/GameplayTasks/Classes/GameplayTaskResource.h
--- a/Engine/Source/Runtime/GameplayTasks/Classes/GameplayTaskResource.h
+++ b/Engine/Source/Runtime/GameplayTasks/Classes/GameplayTaskResource.h
@@ -10,8 +10,8 @@
 
 struct FPropertyChangedEvent;
 
-UCLASS(Abstract, config = "Game", hidedropdown)
-class GAMEPLAYTASKS_API UGameplayTaskResource : public UObject
+UCLASS(Abstract, config = "Game", hidedropdown, MinimalAPI)
+class UGameplayTaskResource : public UObject
 {
 	GENERATED_BODY()
 
@@ -30,7 +30,7 @@ public:
 
 public:
 
-	UGameplayTaskResource(const FObjectInitializer& ObjectInitializer);
+	GAMEPLAYTASKS_API UGameplayTaskResource(const FObjectInitializer& ObjectInitializer);
 
 	uint8 GetResourceID() const
 	{
@@ -48,19 +48,19 @@ public:
 		return RequiredResource->GetDefaultObject<UGameplayTaskResource>()->GetResourceID();
 	}
 
-	virtual void PostInitProperties() override;
+	GAMEPLAYTASKS_API virtual void PostInitProperties() override;
 
 protected:
 #if WITH_EDITOR
-	void PostEditChangeProperty(FPropertyChangedEvent& PropertyChangedEvent);
+	GAMEPLAYTASKS_API void PostEditChangeProperty(FPropertyChangedEvent& PropertyChangedEvent);
 #endif // WITH_EDITOR
 
-	void UpdateAutoResourceID();
+	GAMEPLAYTASKS_API void UpdateAutoResourceID();
 
 #if WITH_GAMEPLAYTASK_DEBUG
 protected:
-	static TArray<FString> ResourceDescriptions;
-	virtual FString GenerateDebugDescription() const;
+	static GAMEPLAYTASKS_API TArray<FString> ResourceDescriptions;
+	GAMEPLAYTASKS_API virtual FString GenerateDebugDescription() const;
 
 public:
 	static FString GetDebugDescription(uint8 ResourceId)
diff --git a/Engine/Source/Runtime/GameplayTasks/Classes/GameplayTasksComponent.h b/Engine/Source/Runtime/GameplayTasks/Classes/GameplayTasksComponent.h
--- a/Engine/Source/Runtime/GameplayTasks/Classes/GameplayTasksComponent.h
+++ b/Engine/Source/Runtime/GameplayTasks/Classes/GameplayTasksComponent.h
@@ -57,8 +57,8 @@ DECLARE_DYNAMIC_MULTICAST_DELEGATE_TwoParams(FOnClaimedResourcesChangeSignature,
 /**
 *	The core ActorComponent for interfacing with the GameplayAbilities System
 */
-UCLASS(ClassGroup = GameplayTasks, hidecategories = (Object, LOD, Lighting, Transform, Sockets, TextureStreaming), editinlinenew, meta = (BlueprintSpawnableComponent))
-class GAMEPLAYTASKS_API UGameplayTasksComponent : public UActorComponent, public IGameplayTaskOwnerInterface, public IVisualLoggerDebugSnapshotInterface
+UCLASS(ClassGroup = GameplayTasks, hidecategories = (Object, LOD, Lighting, Transform, Sockets, TextureStreaming), editinlinenew, meta = (BlueprintSpawnableComponent), MinimalAPI)
+class UGameplayTasksComponent : public UActorComponent, public IGameplayTaskOwnerInterface, public IVisualLoggerDebugSnapshotInterface
 {
 	GENERATED_BODY()
 
@@ -83,16 +83,16 @@ protected:
 	}
 
 	UE_DEPRECATED(5.1, "This will be removed in future versions. Use AddSimulatedTask or RemoveSimulatedTask to modify the array")
-	TArray<TObjectPtr<UGameplayTask>>& GetSimulatedTasks_Mutable();
+	GAMEPLAYTASKS_API TArray<TObjectPtr<UGameplayTask>>& GetSimulatedTasks_Mutable();
 
 	/** Remove all current tasks and register the one in the passed array. It's optimal to use Add/Remove of single tasks directly if possible.*/
-	void SetSimulatedTasks(const TArray<UGameplayTask*>& NewSimulatedTasks);
+	GAMEPLAYTASKS_API void SetSimulatedTasks(const TArray<UGameplayTask*>& NewSimulatedTasks);
 
 	/** Add a new simulated task. Returns true if the task was added to the list. Returns false if the task was already registered. */
-	bool AddSimulatedTask(UGameplayTask* NewTask);
+	GAMEPLAYTASKS_API bool AddSimulatedTask(UGameplayTask* NewTask);
 
 	/** Remove an existing simulated task */
-	void RemoveSimulatedTask(UGameplayTask* NewTask);
+	GAMEPLAYTASKS_API void RemoveSimulatedTask(UGameplayTask* NewTask);
 
 	UPROPERTY()
 	TArray<TObjectPtr<UGameplayTask>> TaskPriorityQueue;
@@ -115,34 +115,34 @@ public:
 	UPROPERTY(BlueprintReadWrite, Category = "Gameplay Tasks")
 	FOnClaimedResourcesChangeSignature OnClaimedResourcesChange;
 
-	UGameplayTasksComponent(const FObjectInitializer& ObjectInitializer);
+	GAMEPLAYTASKS_API UGameplayTasksComponent(const FObjectInitializer& ObjectInitializer);
 	
 	UFUNCTION()
-	void OnRep_SimulatedTasks(const TArray<UGameplayTask*>& PreviousSimulatedTasks);
+	GAMEPLAYTASKS_API void OnRep_SimulatedTasks(const TArray<UGameplayTask*>& PreviousSimulatedTasks);
 
-	virtual void GetLifetimeReplicatedProps(TArray< FLifetimeProperty >& OutLifetimeProps) const override;
-	virtual bool ReplicateSubobjects(UActorChannel *Channel, FOutBunch* Bunch, FReplicationFlags* RepFlags) override;
+	GAMEPLAYTASKS_API virtual void GetLifetimeReplicatedProps(TArray< FLifetimeProperty >& OutLifetimeProps) const override;
+	GAMEPLAYTASKS_API virtual bool ReplicateSubobjects(UActorChannel *Channel, FOutBunch* Bunch, FReplicationFlags* RepFlags) override;
 
-	virtual void TickComponent(float DeltaTime, enum ELevelTick TickType, FActorComponentTickFunction* ThisTickFunction) override;
-	void UpdateShouldTick();
+	GAMEPLAYTASKS_API virtual void TickComponent(float DeltaTime, enum ELevelTick TickType, FActorComponentTickFunction* ThisTickFunction) override;
+	GAMEPLAYTASKS_API void UpdateShouldTick();
 
 	/** retrieves information whether this component should be ticking taken current
 	*	activity into consideration*/
-	virtual bool GetShouldTick() const;
+	GAMEPLAYTASKS_API virtual bool GetShouldTick() const;
 	
 	/** processes the task and figures out if it should get triggered instantly or wait
 	 *	based on task's RequiredResources, Priority and ResourceOverlapPolicy */
-	void AddTaskReadyForActivation(UGameplayTask& NewTask);
+	GAMEPLAYTASKS_API void AddTaskReadyForActivation(UGameplayTask& NewTask);
 
-	void RemoveResourceConsumingTask(UGameplayTask& Task);
-	void EndAllResourceConsumingTasksOwnedBy(const IGameplayTaskOwnerInterface& TaskOwner);
+	GAMEPLAYTASKS_API void RemoveResourceConsumingTask(UGameplayTask& Task);
+	GAMEPLAYTASKS_API void EndAllResourceConsumingTasksOwnedBy(const IGameplayTaskOwnerInterface& TaskOwner);
 
-	bool FindAllResourceConsumingTasksOwnedBy(const IGameplayTaskOwnerInterface& TaskOwner, TArray<UGameplayTask*>& FoundTasks) const;
+	GAMEPLAYTASKS_API bool FindAllResourceConsumingTasksOwnedBy(const IGameplayTaskOwnerInterface& TaskOwner, TArray<UGameplayTask*>& FoundTasks) const;
 	
 	/** finds first resource-consuming task of given name */
-	UGameplayTask* FindResourceConsumingTaskByName(const FName TaskInstanceName) const;
+	GAMEPLAYTASKS_API UGameplayTask* FindResourceConsumingTaskByName(const FName TaskInstanceName) const;
 
-	bool HasActiveTasks(UClass* TaskClass) const;
+	GAMEPLAYTASKS_API bool HasActiveTasks(UClass* TaskClass) const;
 
 	FORCEINLINE FGameplayResourceSet GetCurrentlyUsedResources() const { return CurrentlyClaimedResources; }
 
@@ -150,33 +150,33 @@ public:
 	virtual UGameplayTasksComponent* GetGameplayTasksComponent(const UGameplayTask& Task) const { return const_cast<UGameplayTasksComponent*>(this); }
 	virtual AActor* GetGameplayTaskOwner(const UGameplayTask* Task) const override { return GetOwner(); }
 	virtual AActor* GetGameplayTaskAvatar(const UGameplayTask* Task) const override { return GetOwner(); }
-	virtual void OnGameplayTaskActivated(UGameplayTask& Task) override;
-	virtual void OnGameplayTaskDeactivated(UGameplayTask& Task) override;
+	GAMEPLAYTASKS_API virtual void OnGameplayTaskActivated(UGameplayTask& Task) override;
+	GAMEPLAYTASKS_API virtual void OnGameplayTaskDeactivated(UGameplayTask& Task) override;
 	// END IGameplayTaskOwnerInterface
 
 	// ActorComponent overrides
-	virtual void ReadyForReplication() override;
+	GAMEPLAYTASKS_API virtual void ReadyForReplication() override;
 
 	UFUNCTION(BlueprintCallable, DisplayName="Run Gameplay Task", meta=(ScriptName="RunGameplayTask"), Category = "Gameplay Tasks", meta = (AutoCreateRefTerm = "AdditionalRequiredResources, AdditionalClaimedResources", AdvancedDisplay = "AdditionalRequiredResources, AdditionalClaimedResources"))
-	static EGameplayTaskRunResult K2_RunGameplayTask(TScriptInterface<IGameplayTaskOwnerInterface> TaskOwner, UGameplayTask* Task, uint8 Priority, TArray<TSubclassOf<UGameplayTaskResource> > AdditionalRequiredResources, TArray<TSubclassOf<UGameplayTaskResource> > AdditionalClaimedResources);
+	static GAMEPLAYTASKS_API EGameplayTaskRunResult K2_RunGameplayTask(TScriptInterface<IGameplayTaskOwnerInterface> TaskOwner, UGameplayTask* Task, uint8 Priority, TArray<TSubclassOf<UGameplayTaskResource> > AdditionalRequiredResources, TArray<TSubclassOf<UGameplayTaskResource> > AdditionalClaimedResources);
 
-	static EGameplayTaskRunResult RunGameplayTask(IGameplayTaskOwnerInterface& TaskOwner, UGameplayTask& Task, uint8 Priority, FGameplayResourceSet AdditionalRequiredResources, FGameplayResourceSet AdditionalClaimedResources);
+	static GAMEPLAYTASKS_API EGameplayTaskRunResult RunGameplayTask(IGameplayTaskOwnerInterface& TaskOwner, UGameplayTask& Task, uint8 Priority, FGameplayResourceSet AdditionalRequiredResources, FGameplayResourceSet AdditionalClaimedResources);
 	
 #if WITH_GAMEPLAYTASK_DEBUG
-	FString GetTickingTasksDescription() const;
-	FString GetKnownTasksDescription() const;
-	FString GetTasksPriorityQueueDescription() const;
-	static FString GetTaskStateName(EGameplayTaskState Value);
+	GAMEPLAYTASKS_API FString GetTickingTasksDescription() const;
+	GAMEPLAYTASKS_API FString GetKnownTasksDescription() const;
+	GAMEPLAYTASKS_API FString GetTasksPriorityQueueDescription() const;
+	static GAMEPLAYTASKS_API FString GetTaskStateName(EGameplayTaskState Value);
 #endif // WITH_GAMEPLAYTASK_DEBUG
 	using GameplayTaskContainerType = decltype(TickingTasks);
-	GameplayTaskContainerType::TConstIterator GetTickingTaskIterator() const;
-	GameplayTaskContainerType::TConstIterator GetKnownTaskIterator() const;
-	GameplayTaskContainerType::TConstIterator GetPriorityQueueIterator() const;
-	GameplayTaskContainerType::TConstIterator GetSimulatedTaskIterator() const;
+	GAMEPLAYTASKS_API GameplayTaskContainerType::TConstIterator GetTickingTaskIterator() const;
+	GAMEPLAYTASKS_API GameplayTaskContainerType::TConstIterator GetKnownTaskIterator() const;
+	GAMEPLAYTASKS_API GameplayTaskContainerType::TConstIterator GetPriorityQueueIterator() const;
+	GAMEPLAYTASKS_API GameplayTaskContainerType::TConstIterator GetSimulatedTaskIterator() const;
 
 #if ENABLE_VISUAL_LOG
-	virtual void GrabDebugSnapshot(FVisualLogEntry* Snapshot) const override;
-	void DescribeSelfToVisLog(struct FVisualLogEntry* Snapshot) const;
+	GAMEPLAYTASKS_API virtual void GrabDebugSnapshot(FVisualLogEntry* Snapshot) const override;
+	GAMEPLAYTASKS_API void DescribeSelfToVisLog(struct FVisualLogEntry* Snapshot) const;
 #endif // ENABLE_VISUAL_LOG
 
 protected:
@@ -189,22 +189,22 @@ protected:
 		UGameplayTasksComponent* Owner;
 	};
 
-	void RequestTicking();
-	void ProcessTaskEvents();
-	void UpdateTaskActivations();
+	GAMEPLAYTASKS_API void RequestTicking();
+	GAMEPLAYTASKS_API void ProcessTaskEvents();
+	GAMEPLAYTASKS_API void UpdateTaskActivations();
 
-	void SetCurrentlyClaimedResources(FGameplayResourceSet NewClaimedSet);
+	GAMEPLAYTASKS_API void SetCurrentlyClaimedResources(FGameplayResourceSet NewClaimedSet);
 
 private:
 	/** called when a task gets ended with an external call, meaning not coming from UGameplayTasksComponent mechanics */
-	void OnTaskEnded(UGameplayTask& Task);
+	GAMEPLAYTASKS_API void OnTaskEnded(UGameplayTask& Task);
 
-	void AddTaskToPriorityQueue(UGameplayTask& NewTask);
-	void RemoveTaskFromPriorityQueue(UGameplayTask& Task);
+	GAMEPLAYTASKS_API void AddTaskToPriorityQueue(UGameplayTask& NewTask);
+	GAMEPLAYTASKS_API void RemoveTaskFromPriorityQueue(UGameplayTask& Task);
 
 	FORCEINLINE bool CanProcessEvents() const { return !bInEventProcessingInProgress && (EventLockCounter == 0); }
 
-	void SetSimulatedTasksNetDirty();
+	GAMEPLAYTASKS_API void SetSimulatedTasksNetDirty();
 
 	/** Tasks that run on simulated proxies */
 	UPROPERTY(ReplicatedUsing = OnRep_SimulatedTasks)
diff --git a/Engine/Source/Runtime/GameplayTasks/Classes/Tasks/GameplayTask_ClaimResource.h b/Engine/Source/Runtime/GameplayTasks/Classes/Tasks/GameplayTask_ClaimResource.h
--- a/Engine/Source/Runtime/GameplayTasks/Classes/Tasks/GameplayTask_ClaimResource.h
+++ b/Engine/Source/Runtime/GameplayTasks/Classes/Tasks/GameplayTask_ClaimResource.h
@@ -11,19 +11,19 @@
 #include "GameplayTaskResource.h"
 #include "GameplayTask_ClaimResource.generated.h"
 
-UCLASS(BlueprintType)
-class GAMEPLAYTASKS_API UGameplayTask_ClaimResource : public UGameplayTask
+UCLASS(BlueprintType, MinimalAPI)
+class UGameplayTask_ClaimResource : public UGameplayTask
 {
 	GENERATED_BODY()
 public:
-	UGameplayTask_ClaimResource(const FObjectInitializer& ObjectInitializer = FObjectInitializer::Get());
+	GAMEPLAYTASKS_API UGameplayTask_ClaimResource(const FObjectInitializer& ObjectInitializer = FObjectInitializer::Get());
 
 	UFUNCTION(BlueprintCallable, Category = "AI|Tasks", meta = (AdvancedDisplay = "Priority, TaskInstanceName"))
-	static UGameplayTask_ClaimResource* ClaimResource(TScriptInterface<IGameplayTaskOwnerInterface> InTaskOwner, TSubclassOf<UGameplayTaskResource> ResourceClass, const uint8 Priority = 192, const FName TaskInstanceName = NAME_None);
+	static GAMEPLAYTASKS_API UGameplayTask_ClaimResource* ClaimResource(TScriptInterface<IGameplayTaskOwnerInterface> InTaskOwner, TSubclassOf<UGameplayTaskResource> ResourceClass, const uint8 Priority = 192, const FName TaskInstanceName = NAME_None);
 
 	UFUNCTION(BlueprintCallable, Category = "AI|Tasks", meta = (AdvancedDisplay = "Priority, TaskInstanceName"))
-	static UGameplayTask_ClaimResource* ClaimResources(TScriptInterface<IGameplayTaskOwnerInterface> InTaskOwner, TArray<TSubclassOf<UGameplayTaskResource> > ResourceClasses, const uint8 Priority = 192, const FName TaskInstanceName = NAME_None);
+	static GAMEPLAYTASKS_API UGameplayTask_ClaimResource* ClaimResources(TScriptInterface<IGameplayTaskOwnerInterface> InTaskOwner, TArray<TSubclassOf<UGameplayTaskResource> > ResourceClasses, const uint8 Priority = 192, const FName TaskInstanceName = NAME_None);
 
-	static UGameplayTask_ClaimResource* ClaimResource(IGameplayTaskOwnerInterface& InTaskOwner, const TSubclassOf<UGameplayTaskResource> ResourceClass, const uint8 Priority = FGameplayTasks::DefaultPriority, const FName TaskInstanceName = NAME_None);
-	static UGameplayTask_ClaimResource* ClaimResources(IGameplayTaskOwnerInterface& InTaskOwner, const TArray<TSubclassOf<UGameplayTaskResource> >& ResourceClasses, const uint8 Priority = FGameplayTasks::DefaultPriority, const FName TaskInstanceName = NAME_None);
+	static GAMEPLAYTASKS_API UGameplayTask_ClaimResource* ClaimResource(IGameplayTaskOwnerInterface& InTaskOwner, const TSubclassOf<UGameplayTaskResource> ResourceClass, const uint8 Priority = FGameplayTasks::DefaultPriority, const FName TaskInstanceName = NAME_None);
+	static GAMEPLAYTASKS_API UGameplayTask_ClaimResource* ClaimResources(IGameplayTaskOwnerInterface& InTaskOwner, const TArray<TSubclassOf<UGameplayTaskResource> >& ResourceClasses, const uint8 Priority = FGameplayTasks::DefaultPriority, const FName TaskInstanceName = NAME_None);
 };
diff --git a/Engine/Source/Runtime/GeometryCore/Public/CompGeom/ConvexDecomposition3.h b/Engine/Source/Runtime/GeometryCore/Public/CompGeom/ConvexDecomposition3.h
--- a/Engine/Source/Runtime/GeometryCore/Public/CompGeom/ConvexDecomposition3.h
+++ b/Engine/Source/Runtime/GeometryCore/Public/CompGeom/ConvexDecomposition3.h
@@ -72,7 +72,7 @@ struct FNegativeSpaceSampleSettings
 };
 
 // Define a volume with a set of spheres
-class GEOMETRYCORE_API FSphereCovering
+class FSphereCovering
 {
 public:
 
@@ -93,7 +93,7 @@ public:
 
 	// Add spheres covering the negative space of the given fast winding tree
 	// @return true if any spheres were added
-	bool AddNegativeSpace(const TFastWindingTree<FDynamicMesh3>& Spatial, const FNegativeSpaceSampleSettings& Settings);
+	GEOMETRYCORE_API bool AddNegativeSpace(const TFastWindingTree<FDynamicMesh3>& Spatial, const FNegativeSpaceSampleSettings& Settings);
 
 	void Reset()
 	{
@@ -110,7 +110,7 @@ private:
 
 
 
-class GEOMETRYCORE_API FConvexDecomposition3
+class FConvexDecomposition3
 {
 
 public:
@@ -132,9 +132,9 @@ public:
 	 * @param HullVertices		Function from Hull Index, Vertex Index -> Hull Vertex Position
 	 * @param Proximity			All the local proximities for the hulls. Hulls will not be merged unless they are connected by this proximity graph.
 	 */
-	void InitializeFromHulls(int32 NumHulls, TFunctionRef<double(int32)> HullVolumes, TFunctionRef<int32(int32)> HullNumVertices, TFunctionRef<FVector3d(int32, int32)> HullVertices, TArrayView<const TPair<int32, int32>> Proximity);
+	GEOMETRYCORE_API void InitializeFromHulls(int32 NumHulls, TFunctionRef<double(int32)> HullVolumes, TFunctionRef<int32(int32)> HullNumVertices, TFunctionRef<FVector3d(int32, int32)> HullVertices, TArrayView<const TPair<int32, int32>> Proximity);
 
-	void InitializeFromMesh(const FDynamicMesh3& SourceMesh, bool bMergeEdges);
+	GEOMETRYCORE_API void InitializeFromMesh(const FDynamicMesh3& SourceMesh, bool bMergeEdges);
 
 	/**
 	 * Initialize convex decomposition with a triangle index mesh
@@ -143,7 +143,7 @@ public:
 	 * @param bMergeEdges		Whether to attempt to weld matching edges before computing the convex hull; this can help the convex decomposition find better cutting planes for meshes that have boundaries e.g. due to seams
 	 * @param FaceVertexOffset	Indices from the Faces array are optionally offset by this value.  Useful e.g. to take slices of the multi-geometry vertex and face buffers of FGeometryCollection.
 	 */
-	void InitializeFromIndexMesh(TArrayView<const FVector3f> Vertices, TArrayView<const FIntVector> Faces, bool bMergeEdges, int32 FaceVertexOffset = 0);
+	GEOMETRYCORE_API void InitializeFromIndexMesh(TArrayView<const FVector3f> Vertices, TArrayView<const FIntVector> Faces, bool bMergeEdges, int32 FaceVertexOffset = 0);
 
 	//
 	// Settings
@@ -211,7 +211,7 @@ public:
 	// Split the worst convex part, and return the increase in the total number of convex parts after splitting (can be more than 1 if result has multiple separate connected components)
 	// Note: could return 0 if no splits were possible
 	// @param bCanSkipUnreliableGeoVolumes		if true, don't split hulls where we have questionable geometry volume results, unless there is no hull with good geometry volume results
-	int32 SplitWorst(bool bCanSkipUnreliableGeoVolumes = false, double ErrorTolerance = 0.0);
+	GEOMETRYCORE_API int32 SplitWorst(bool bCanSkipUnreliableGeoVolumes = false, double ErrorTolerance = 0.0);
 
 	// Merge the pairs of convex hulls in the decomposition that will least increase the error.  Intermediate results can be used across merges, so it is best to do all merges in one call.
 	// Note: A future version of this function may replace NumOutputHulls with MaxOutputHulls, but this version keeps both parameters for compatibility / consistent behavior.
@@ -226,7 +226,7 @@ public:
 	// @param OptionalNegativeSpace Optional representation of negative space that should not be covered by merges
 	// @param OptionalNegativeSpaceTransform Optional transform from space of the convex hulls into the space of the sphere covering; if not provided, assume spheres are in the same coordinate space as the hulls
 	// @return						The number of merges performed
-	int32 MergeBest(int32 TargetNumParts, double ErrorTolerance = 0, double MinThicknessTolerance = 0, bool bAllowCompact = true, bool bRequireHullTriangles = false, int MaxOutputHulls = -1,
+	GEOMETRYCORE_API int32 MergeBest(int32 TargetNumParts, double ErrorTolerance = 0, double MinThicknessTolerance = 0, bool bAllowCompact = true, bool bRequireHullTriangles = false, int MaxOutputHulls = -1,
 		const FSphereCovering* OptionalNegativeSpace = nullptr, const FTransform* OptionalTransformIntoNegativeSpace = nullptr);
 
 	// simple helper to convert an error tolerance expressed in world space to a local-space volume tolerance
@@ -474,7 +474,7 @@ public:
 
 	// Helper function to delete a proximity relationship, and fix index references to account for the new array ordering
 	// @param bDeleteMapReferences		If true, also update references in the DecompositionToProximity map
-	void DeleteProximity(int32 ProxIdx, bool bDeleteMapReferences);
+	GEOMETRYCORE_API void DeleteProximity(int32 ProxIdx, bool bDeleteMapReferences);
 
 	// Helper function to update all proximity data after splitting an FConvexPart apart
 	// @param SplitIdx				The index of the FConvexPart that has been split into multiple parts.  Now contains the first new part.
@@ -482,7 +482,7 @@ public:
 	// @param CutPlane				The plane used for the split
 	// @param SecondSideIdxStart	The first index in the Decomposition array of new FConvexParts that were on the opposite side of the plane from SplitIdx.  Will be different from NewIdxStart only if the first part was split into multiple components.
 	// @param OrigHullVolume		The volume of the convex hull *before* the split was performed -- used to precompute the merge cost.
-	void UpdateProximitiesAfterSplit(int32 SplitIdx, int32 NewIdxStart, FPlane3d CutPlane, int32 SecondSideIdxStart, double OrigHullVolume);
+	GEOMETRYCORE_API void UpdateProximitiesAfterSplit(int32 SplitIdx, int32 NewIdxStart, FPlane3d CutPlane, int32 SecondSideIdxStart, double OrigHullVolume);
 };
 
 } // end namespace UE::Geometry
diff --git a/Engine/Source/Runtime/GeometryCore/Public/CompGeom/Delaunay2.h b/Engine/Source/Runtime/GeometryCore/Public/CompGeom/Delaunay2.h
--- a/Engine/Source/Runtime/GeometryCore/Public/CompGeom/Delaunay2.h
+++ b/Engine/Source/Runtime/GeometryCore/Public/CompGeom/Delaunay2.h
@@ -32,7 +32,7 @@ using namespace UE::Math;
 // Internal representation of mesh connectivity; not exposed to interface
 struct FDelaunay2Connectivity;
 
-class GEOMETRYCORE_API FDelaunay2
+class FDelaunay2
 {
 public:
 
@@ -78,8 +78,8 @@ public:
 	 *
 	 * @return false if triangulation failed
 	 */
-	bool Triangulate(TArrayView<const TVector2<double>> Vertices, TArrayView<const FIndex2i> Edges = TArrayView<const FIndex2i>());
-	bool Triangulate(TArrayView<const TVector2<float>> Vertices, TArrayView<const FIndex2i> Edges = TArrayView<const FIndex2i>());
+	GEOMETRYCORE_API bool Triangulate(TArrayView<const TVector2<double>> Vertices, TArrayView<const FIndex2i> Edges = TArrayView<const FIndex2i>());
+	GEOMETRYCORE_API bool Triangulate(TArrayView<const TVector2<float>> Vertices, TArrayView<const FIndex2i> Edges = TArrayView<const FIndex2i>());
 
 	// Triangulate a polygon, and optionally pass back the resulting triangles
 	// Uses the 'solid' fill mode, and fills the polygon regardless of the edge orientation
@@ -162,8 +162,8 @@ public:
 	 * @return true if edges were successfully inserted.
 	 * Note: Will not detect failure due to intersection of constrained edges unless member bValidateEdges == true.
 	 */
-	bool ConstrainEdges(TArrayView<const TVector2<double>> Vertices, TArrayView<const FIndex2i> Edges);
-	bool ConstrainEdges(TArrayView<const TVector2<float>> Vertices, TArrayView<const FIndex2i> Edges);
+	GEOMETRYCORE_API bool ConstrainEdges(TArrayView<const TVector2<double>> Vertices, TArrayView<const FIndex2i> Edges);
+	GEOMETRYCORE_API bool ConstrainEdges(TArrayView<const TVector2<float>> Vertices, TArrayView<const FIndex2i> Edges);
 
 	// TODO: Support incremental vertex insertion
 	// Update the triangulation incrementally, assuming Vertices are unchanged before FirstNewIndex, and nothing after FirstNewIndex has been inserted yet
@@ -172,10 +172,10 @@ public:
 
 	// Get the triangulation as an array of triangles
 	// Note: This creates a new array each call, because the internal data structure does not have a triangle array
-	TArray<FIndex3i> GetTriangles() const;
+	GEOMETRYCORE_API TArray<FIndex3i> GetTriangles() const;
 
 	// Get the triangulation as an array with a corresponding adjacency array, indicating the adjacent triangle on each triangle edge (-1 if no adjacent triangle)
-	void GetTrianglesAndAdjacency(TArray<FIndex3i>& Triangles, TArray<FIndex3i>& Adjacency) const;
+	GEOMETRYCORE_API void GetTrianglesAndAdjacency(TArray<FIndex3i>& Triangles, TArray<FIndex3i>& Adjacency) const;
 
 	/**
 	 * Return the triangles that are inside the given edges, removing the outer boundary triangles
@@ -202,7 +202,7 @@ public:
 	 * @return							true if the result was well defined and consistent; false otherwise. Solid fill mode always has a well defined result; Winding-number-based fills do not if the edges have open spans.
 	 *									Note the triangulation will still be filled by best-effort even if the function returns false.
 	 */
-	bool GetFilledTriangles(TArray<FIndex3i>& TrianglesOut, TArrayView<const FIndex2i> Edges, EFillMode FillMode = EFillMode::PositiveWinding) const;
+	GEOMETRYCORE_API bool GetFilledTriangles(TArray<FIndex3i>& TrianglesOut, TArrayView<const FIndex2i> Edges, EFillMode FillMode = EFillMode::PositiveWinding) const;
 
 	/**
 	 * Get (by reference) the triangles that are inside the given edges, using a generalized winding number method to determine which triangles are inside
@@ -218,8 +218,8 @@ public:
 	 * @return							true if the result was well defined and consistent; false otherwise. Solid fill mode always has a well defined result; Winding-number-based fills do not if the edges have open spans.
 	 *									Note the triangulation will still be filled by best-effort even if the function returns false.
 	 */
-	bool GetFilledTrianglesGeneralizedWinding(TArray<FIndex3i>& TrianglesOut, TArrayView<const TVector2<double>> Vertices, TArrayView<const FIndex2i> Edges, EFillMode FillMode = EFillMode::PositiveWinding) const;
-	bool GetFilledTrianglesGeneralizedWinding(TArray<FIndex3i>& TrianglesOut, TArrayView<const TVector2<float>> Vertices, TArrayView<const FIndex2i> Edges, EFillMode FillMode = EFillMode::PositiveWinding) const;
+	GEOMETRYCORE_API bool GetFilledTrianglesGeneralizedWinding(TArray<FIndex3i>& TrianglesOut, TArrayView<const TVector2<double>> Vertices, TArrayView<const FIndex2i> Edges, EFillMode FillMode = EFillMode::PositiveWinding) const;
+	GEOMETRYCORE_API bool GetFilledTrianglesGeneralizedWinding(TArray<FIndex3i>& TrianglesOut, TArrayView<const TVector2<float>> Vertices, TArrayView<const FIndex2i> Edges, EFillMode FillMode = EFillMode::PositiveWinding) const;
 
 	/**
 	 * Get (by reference) the triangles that are inside the given edges, removing the outside-boundary triangles and the inside-hole triangles
@@ -230,7 +230,7 @@ public:
 	 * @return								true if any result was successfully computed (including an empty result).
 	 *										Currently only returns false if Triangulate() has not been called yet.
 	 */
-	bool GetFilledTriangles(TArray<FIndex3i>& TrianglesOut, TArrayView<const FIndex2i> BoundaryEdges, TArrayView<const FIndex2i> HoleEdges) const;
+	GEOMETRYCORE_API bool GetFilledTriangles(TArray<FIndex3i>& TrianglesOut, TArrayView<const FIndex2i> BoundaryEdges, TArrayView<const FIndex2i> HoleEdges) const;
 
 	// @return true if this is a constrained Delaunay triangulation
 	bool IsConstrained() const
@@ -239,23 +239,23 @@ public:
 	}
 
 	// @return true if triangulation is Delaunay, useful for validating results (note: likely to be false if edges are constrained)
-	bool IsDelaunay(TArrayView<const FVector2f> Vertices, TArrayView<const FIndex2i> SkipEdges = TArrayView<const FIndex2i>()) const;
-	bool IsDelaunay(TArrayView<const FVector2d> Vertices, TArrayView<const FIndex2i> SkipEdges = TArrayView<const FIndex2i>()) const;
+	GEOMETRYCORE_API bool IsDelaunay(TArrayView<const FVector2f> Vertices, TArrayView<const FIndex2i> SkipEdges = TArrayView<const FIndex2i>()) const;
+	GEOMETRYCORE_API bool IsDelaunay(TArrayView<const FVector2d> Vertices, TArrayView<const FIndex2i> SkipEdges = TArrayView<const FIndex2i>()) const;
 
 	// @return true if triangulation has the given edges, useful for validating results
-	bool HasEdges(TArrayView<const FIndex2i> Edges) const;
+	GEOMETRYCORE_API bool HasEdges(TArrayView<const FIndex2i> Edges) const;
 
 	// Remap any references to duplicate vertices to only reference the vertices used in the triangulation
-	void FixDuplicatesOnEdge(FIndex2i& Edge);
+	GEOMETRYCORE_API void FixDuplicatesOnEdge(FIndex2i& Edge);
 
 	// @return true if the edge is in the triangulation
-	bool HasEdge(const FIndex2i& Edge, bool bRemapDuplicates);
+	GEOMETRYCORE_API bool HasEdge(const FIndex2i& Edge, bool bRemapDuplicates);
 
 	// @return true if the input had duplicate vertices, and so not all vertices could be used in the triangulation.
-	bool HasDuplicates() const;
+	GEOMETRYCORE_API bool HasDuplicates() const;
 
 	// @return if bAutomaticallyFixEdgesToDuplicateVertices was set during triangulation, will return the remapped vertex index for any duplicate vertices, or Index if the vertex was not remapped.
-	int32 RemapIfDuplicate(int32 Index) const;
+	GEOMETRYCORE_API int32 RemapIfDuplicate(int32 Index) const;
 
 
 
@@ -267,8 +267,8 @@ public:
 	 * @param ExpandBounds		Amount to expand the clipping bounds beyond the Bounds argument (or the bounding box of non-boundary Voronoi cells, if Bounds was empty)
 	 * @return					The cells of each Voronoi site, or an empty array if the Triangulation was not yet computed
 	 */
-	TArray<TArray<FVector2d>> GetVoronoiCells(TArrayView<const FVector2d> Vertices, bool bIncludeBoundary = false, FAxisAlignedBox2d ClipBounds = FAxisAlignedBox2d::Empty(), double ExpandBounds = 0.0) const;
-	TArray<TArray<FVector2f>> GetVoronoiCells(TArrayView<const FVector2f> Vertices, bool bIncludeBoundary = false, FAxisAlignedBox2f ClipBounds = FAxisAlignedBox2f::Empty(), float ExpandBounds = 0.0f) const;
+	GEOMETRYCORE_API TArray<TArray<FVector2d>> GetVoronoiCells(TArrayView<const FVector2d> Vertices, bool bIncludeBoundary = false, FAxisAlignedBox2d ClipBounds = FAxisAlignedBox2d::Empty(), double ExpandBounds = 0.0) const;
+	GEOMETRYCORE_API TArray<TArray<FVector2f>> GetVoronoiCells(TArrayView<const FVector2f> Vertices, bool bIncludeBoundary = false, FAxisAlignedBox2f ClipBounds = FAxisAlignedBox2f::Empty(), float ExpandBounds = 0.0f) const;
 
 	/**
 	 * Compute Voronoi diagram cells
diff --git a/Engine/Source/Runtime/GeometryCore/Public/CompGeom/Delaunay3.h b/Engine/Source/Runtime/GeometryCore/Public/CompGeom/Delaunay3.h
--- a/Engine/Source/Runtime/GeometryCore/Public/CompGeom/Delaunay3.h
+++ b/Engine/Source/Runtime/GeometryCore/Public/CompGeom/Delaunay3.h
@@ -19,7 +19,7 @@ using namespace UE::Math;
 // Internal representation of mesh connectivity; not exposed to interface
 struct FDelaunay3Connectivity;
 
-class GEOMETRYCORE_API FDelaunay3
+class FDelaunay3
 {
 public:
 
@@ -36,17 +36,17 @@ public:
 	 *
 	 * @return false if triangulation failed
 	 */
-	bool Triangulate(TArrayView<const TVector<double>> Vertices);
-	bool Triangulate(TArrayView<const TVector<float>> Vertices);
+	GEOMETRYCORE_API bool Triangulate(TArrayView<const TVector<double>> Vertices);
+	GEOMETRYCORE_API bool Triangulate(TArrayView<const TVector<float>> Vertices);
 
 	// Get the result as an array of tetrahedra
 	// Note: This creates a new array each call, because the internal data structure does not have a tetrahedra array
-	TArray<FIndex4i> GetTetrahedraAsFIndex4i(bool bReverseOrientation = false) const;
-	TArray<FIntVector4> GetTetrahedra(bool bReverseOrientation = false) const;
+	GEOMETRYCORE_API TArray<FIndex4i> GetTetrahedraAsFIndex4i(bool bReverseOrientation = false) const;
+	GEOMETRYCORE_API TArray<FIntVector4> GetTetrahedra(bool bReverseOrientation = false) const;
 
 	// @return true if triangulation is Delaunay, useful for validating results (note: likely to be false if edges are constrained)
-	bool IsDelaunay(TArrayView<const FVector3f> Vertices) const;
-	bool IsDelaunay(TArrayView<const FVector3d> Vertices) const;
+	GEOMETRYCORE_API bool IsDelaunay(TArrayView<const FVector3f> Vertices) const;
+	GEOMETRYCORE_API bool IsDelaunay(TArrayView<const FVector3d> Vertices) const;
 
 protected:
 	TPimplPtr<FDelaunay3Connectivity> Connectivity;
diff --git a/Engine/Source/Runtime/GeometryCore/Public/DynamicMesh/ColliderMesh.h b/Engine/Source/Runtime/GeometryCore/Public/DynamicMesh/ColliderMesh.h
--- a/Engine/Source/Runtime/GeometryCore/Public/DynamicMesh/ColliderMesh.h
+++ b/Engine/Source/Runtime/GeometryCore/Public/DynamicMesh/ColliderMesh.h
@@ -39,7 +39,7 @@ class FDynamicMesh3;
  * FColliderMesh (effectively) implements the TTriangleMeshAdapter interface, and so can be used 
  * anywhere a TTriangleMeshAdapter could be used (eg with TMeshQueries functions). 
  */
-class GEOMETRYCORE_API FColliderMesh
+class FColliderMesh
 {
 public:
 	struct FBuildOptions
@@ -60,12 +60,12 @@ public:
 	};
 
 public:
-	FColliderMesh();
-	FColliderMesh(const FDynamicMesh3& SourceMesh, const FBuildOptions& BuildOptions = FBuildOptions());
+	GEOMETRYCORE_API FColliderMesh();
+	GEOMETRYCORE_API FColliderMesh(const FDynamicMesh3& SourceMesh, const FBuildOptions& BuildOptions = FBuildOptions());
 
-	void Initialize(const FDynamicMesh3& SourceMesh, const FBuildOptions& BuildOptions = FBuildOptions());
+	GEOMETRYCORE_API void Initialize(const FDynamicMesh3& SourceMesh, const FBuildOptions& BuildOptions = FBuildOptions());
 
-	void Reset(bool bReleaseMemory);
+	GEOMETRYCORE_API void Reset(bool bReleaseMemory);
 
 	// mesh API required for TMeshAABBTree3
 	bool IsTriangle(int32 TriangleID) const { return TriangleID >= 0 && TriangleID < Triangles.Num(); }
@@ -94,10 +94,10 @@ public:
 
 	bool SourceWasCompactV() const { return bSourceWasCompactV; }
 	bool HasVertexIDMap() const { return SourceVertexIDs.Num() > 0 && SourceVertexIDs.Num() == Vertices.Num(); }
-	int32 GetSourceVertexID(int32 VertexID) const;
+	GEOMETRYCORE_API int32 GetSourceVertexID(int32 VertexID) const;
 	bool SourceWasCompactT() const { return bSourceWasCompactT; }
 	bool HasTriangleIDMap() const { return SourceTriangleIDs.Num() > 0 && SourceTriangleIDs.Num() == Triangles.Num(); }
-	int32 GetSourceTriangleID(int32 TriangleID) const;
+	GEOMETRYCORE_API int32 GetSourceTriangleID(int32 TriangleID) const;
 
 	// spatial APIs
 
@@ -105,21 +105,21 @@ public:
 	 * @param HitTriangleIDOut TriangleID on collider mesh - use GetSourceTriangleID() to map back to source mesh
 	 * @return true if a triangle was hit and Out parameters are initialized
 	 */
-	bool FindNearestHitTriangle( const FRay3d& Ray, double& RayParameterOut, int& HitTriangleIDOut, FVector3d& BaryCoordsOut ) const;
+	GEOMETRYCORE_API bool FindNearestHitTriangle( const FRay3d& Ray, double& RayParameterOut, int& HitTriangleIDOut, FVector3d& BaryCoordsOut ) const;
 
 	/**
 	 * Find the TriangleID closest to P on the collider mesh, and distance to it, within distance MaxDist, or return InvalidID.
 	 * Use GetSourceTriangleID() to map the TriangleID back to the source mesh.
 	 * Use MeshQueries.TriangleDistance() to get more information.
 	 */
-	int FindNearestTriangle( const FVector3d& Point, double& NearestDistSqrOut, const IMeshSpatial::FQueryOptions& Options = IMeshSpatial::FQueryOptions() ) const;
+	GEOMETRYCORE_API int FindNearestTriangle( const FVector3d& Point, double& NearestDistSqrOut, const IMeshSpatial::FQueryOptions& Options = IMeshSpatial::FQueryOptions() ) const;
 
 
 	/**
 	 * Direct access to the AABBTree pointer inside the ColliderMesh. 
 	 * @warning this function should be avoided if possible, and may be removed in the future. 
 	 */
-	TMeshAABBTree3<FColliderMesh>* GetRawAABBTreeUnsafe();
+	GEOMETRYCORE_API TMeshAABBTree3<FColliderMesh>* GetRawAABBTreeUnsafe();
 
 
 
@@ -143,7 +143,7 @@ protected:
 /**
  * Projection target API wrapper for an FColliderMesh
  */
-class GEOMETRYCORE_API FColliderMeshProjectionTarget : public IOrientedProjectionTarget
+class FColliderMeshProjectionTarget : public IOrientedProjectionTarget
 {
 public:
 	FColliderMesh* ColliderMesh = nullptr;
@@ -161,8 +161,8 @@ public:
 
 	virtual ~FColliderMeshProjectionTarget() {}
 
-	virtual FVector3d Project(const FVector3d& Point, int Identifier = -1) override;
-	virtual FVector3d Project(const FVector3d& Point, FVector3d& ProjectNormalOut, int Identifier = -1) override;
+	GEOMETRYCORE_API virtual FVector3d Project(const FVector3d& Point, int Identifier = -1) override;
+	GEOMETRYCORE_API virtual FVector3d Project(const FVector3d& Point, FVector3d& ProjectNormalOut, int Identifier = -1) override;
 };
 
 
diff --git a/Engine/Source/Runtime/GeometryCore/Public/DynamicMesh/DynamicMeshAttributeSet.h b/Engine/Source/Runtime/GeometryCore/Public/DynamicMesh/DynamicMeshAttributeSet.h
--- a/Engine/Source/Runtime/GeometryCore/Public/DynamicMesh/DynamicMeshAttributeSet.h
+++ b/Engine/Source/Runtime/GeometryCore/Public/DynamicMesh/DynamicMeshAttributeSet.h
@@ -76,19 +76,19 @@ using FDynamicMeshBonePoseAttribute = TDynamicBoneAttributeBase<FDynamicMesh3, F
  * 
  * @todo current internal structure is a work-in-progress
  */
-class GEOMETRYCORE_API FDynamicMeshAttributeSet : public FDynamicMeshAttributeSetBase
+class FDynamicMeshAttributeSet : public FDynamicMeshAttributeSetBase
 {
 public:
-	FDynamicMeshAttributeSet(FDynamicMesh3* Mesh);
+	GEOMETRYCORE_API FDynamicMeshAttributeSet(FDynamicMesh3* Mesh);
 
-	FDynamicMeshAttributeSet(FDynamicMesh3* Mesh, int32 NumUVLayers, int32 NumNormalLayers);
+	GEOMETRYCORE_API FDynamicMeshAttributeSet(FDynamicMesh3* Mesh, int32 NumUVLayers, int32 NumNormalLayers);
 
-	virtual ~FDynamicMeshAttributeSet() override;
+	GEOMETRYCORE_API virtual ~FDynamicMeshAttributeSet() override;
 
-	void Copy(const FDynamicMeshAttributeSet& Copy);
+	GEOMETRYCORE_API void Copy(const FDynamicMeshAttributeSet& Copy);
 
 	/** returns true if the attached overlays/attributes are compact */
-	bool IsCompact();
+	GEOMETRYCORE_API bool IsCompact();
 
 	/**
 	 * Performs a CompactCopy of the attached overlays/attributes.
@@ -97,7 +97,7 @@ public:
 	 * @param CompactMaps Maps indicating how vertices and triangles were changes in the parent
 	 * @param Copy The attribute set to be copied
 	 */
-	void CompactCopy(const FCompactMaps& CompactMaps, const FDynamicMeshAttributeSet& Copy);
+	GEOMETRYCORE_API void CompactCopy(const FCompactMaps& CompactMaps, const FDynamicMeshAttributeSet& Copy);
 
 	/**
 	 * Compacts the attribute set in place
@@ -105,14 +105,14 @@ public:
 	 *
 	 * @param CompactMaps Maps of how the vertices and triangles were compacted in the parent
 	 */
-	void CompactInPlace(const FCompactMaps& CompactMaps);
+	GEOMETRYCORE_API void CompactInPlace(const FCompactMaps& CompactMaps);
 
 
 	/**
 	 * Split all bowtie vertices in all layers
 	 * @param bParallel if true, layers are processed in parallel
 	 */
-	void SplitAllBowties(bool bParallel = true);
+	GEOMETRYCORE_API void SplitAllBowties(bool bParallel = true);
 
 
 	/**
@@ -122,7 +122,7 @@ public:
 	 * If bClearExisting=false, new attributes are added, but existing attributes and data are preserved
 	 * If bDiscardExtraAttributes=true and bClearExisting=false, then attributes in the current set but not in ToMatch are discarded, but existing attributes are not cleared/reset
 	 */
-	void EnableMatchingAttributes(const FDynamicMeshAttributeSet& ToMatch, bool bClearExisting = true, bool bDiscardExtraAttributes = false);
+	GEOMETRYCORE_API void EnableMatchingAttributes(const FDynamicMeshAttributeSet& ToMatch, bool bClearExisting = true, bool bDiscardExtraAttributes = false);
 
 	/** @return the parent mesh for this overlay */
 	const FDynamicMesh3* GetParentMesh() const { return ParentMesh; }
@@ -131,28 +131,28 @@ public:
 
 private:
 	/** @set the parent mesh for this overlay.  Only safe for use during FDynamicMesh move */
-	void Reparent(FDynamicMesh3* NewParent);
+	GEOMETRYCORE_API void Reparent(FDynamicMesh3* NewParent);
 
 public:
 
 	/** @return true if the given edge is a seam edge in any overlay */
-	virtual bool IsSeamEdge(int EdgeID) const;
+	GEOMETRYCORE_API virtual bool IsSeamEdge(int EdgeID) const;
 
 	/** @return true if the given edge is the termination of a seam in any overlay*/
-	virtual bool IsSeamEndEdge(int EdgeID) const;
+	GEOMETRYCORE_API virtual bool IsSeamEndEdge(int EdgeID) const;
 
 	/** @return true if the given edge is a seam edge in any overlay. This version considers tangent seams as a type of normal seam; to consider tangent seams separately, call the 4-argument IsSeamEdge instead.  */
 	UE_DEPRECATED(5.3, "Please instead call the 4 argument version of IsSeamEdge, which distinguishes between tangent and normal seam edges.")
-	virtual bool IsSeamEdge(int EdgeID, bool& bIsUVSeamOut, bool& bIsNormalSeamOut, bool& bIsColorSeamOut) const;
+	GEOMETRYCORE_API virtual bool IsSeamEdge(int EdgeID, bool& bIsUVSeamOut, bool& bIsNormalSeamOut, bool& bIsColorSeamOut) const;
 
 	/** @return true if the given edge is a seam edge in any overlay */
-	virtual bool IsSeamEdge(int EdgeID, bool& bIsUVSeamOut, bool& bIsNormalSeamOut, bool& bIsColorSeamOut, bool& bIsTangentSeamOut) const;
+	GEOMETRYCORE_API virtual bool IsSeamEdge(int EdgeID, bool& bIsUVSeamOut, bool& bIsNormalSeamOut, bool& bIsColorSeamOut, bool& bIsTangentSeamOut) const;
 
 	/** @return true if the given vertex is a seam vertex in any overlay */
-	virtual bool IsSeamVertex(int VertexID, bool bBoundaryIsSeam = true) const;
+	GEOMETRYCORE_API virtual bool IsSeamVertex(int VertexID, bool bBoundaryIsSeam = true) const;
 
 	/** @return true if the given edge is a material ID boundary */
-	virtual bool IsMaterialBoundaryEdge(int EdgeID) const;
+	GEOMETRYCORE_API virtual bool IsMaterialBoundaryEdge(int EdgeID) const;
 
 	//
 	// UV Layers 
@@ -165,7 +165,7 @@ public:
 	}
 
 	/** Set number of UV (2-vector float overlay) layers */
-	virtual void SetNumUVLayers(int Num);
+	GEOMETRYCORE_API virtual void SetNumUVLayers(int Num);
 
 	/** @return the UV layer at the given Index  if exists, else nullptr*/
 	FDynamicMeshUVOverlay* GetUVLayer(int Index)
@@ -210,13 +210,13 @@ public:
 	}
 
 	/** Set number of Normals (3-vector float overlay) layers */
-	virtual void SetNumNormalLayers(int Num);
+	GEOMETRYCORE_API virtual void SetNumNormalLayers(int Num);
 
 	/** Enable Tangents overlays (Tangent = Normal layer 1, Bitangent = Normal layer 2) */
-	void EnableTangents();
+	GEOMETRYCORE_API void EnableTangents();
 
 	/** Disable Tangents overlays */
-	void DisableTangents();
+	GEOMETRYCORE_API void DisableTangents();
 
 	/** @return the Normal layer at the given Index if exists, else nullptr */
 	FDynamicMeshNormalOverlay* GetNormalLayer(int Index)
@@ -291,44 +291,44 @@ public:
 		return ColorLayer.Get();
 	}
 
-	void EnablePrimaryColors();
+	GEOMETRYCORE_API void EnablePrimaryColors();
 	
-	void DisablePrimaryColors();
+	GEOMETRYCORE_API void DisablePrimaryColors();
 
 	//
 	// Polygroup layers
 	//
 
 	/** @return number of Polygroup layers */
-	virtual int32 NumPolygroupLayers() const;
+	GEOMETRYCORE_API virtual int32 NumPolygroupLayers() const;
 
 	/** Set the number of Polygroup layers */
-	virtual void SetNumPolygroupLayers(int32 Num);
+	GEOMETRYCORE_API virtual void SetNumPolygroupLayers(int32 Num);
 
 	/** @return the Polygroup layer at the given Index */
-	FDynamicMeshPolygroupAttribute* GetPolygroupLayer(int Index);
+	GEOMETRYCORE_API FDynamicMeshPolygroupAttribute* GetPolygroupLayer(int Index);
 
 	/** @return the Polygroup layer at the given Index */
-	const FDynamicMeshPolygroupAttribute* GetPolygroupLayer(int Index) const;
+	GEOMETRYCORE_API const FDynamicMeshPolygroupAttribute* GetPolygroupLayer(int Index) const;
 
 	//
 	// Weight layers
 	//
 
 	/** @return number of weight layers */
-	virtual int32 NumWeightLayers() const;
+	GEOMETRYCORE_API virtual int32 NumWeightLayers() const;
 
 	/** Set the number of weight layers */
-	virtual void SetNumWeightLayers(int32 Num);
+	GEOMETRYCORE_API virtual void SetNumWeightLayers(int32 Num);
 
 	/** Remove a weight layer at the specified index */
-	virtual void RemoveWeightLayer(int32 Index);
+	GEOMETRYCORE_API virtual void RemoveWeightLayer(int32 Index);
 
 	/** @return the weight layer at the given Index */
-	FDynamicMeshWeightAttribute* GetWeightLayer(int Index);
+	GEOMETRYCORE_API FDynamicMeshWeightAttribute* GetWeightLayer(int Index);
 
 	/** @return the weight layer at the given Index */
-	const FDynamicMeshWeightAttribute* GetWeightLayer(int Index) const;
+	GEOMETRYCORE_API const FDynamicMeshWeightAttribute* GetWeightLayer(int Index) const;
 
 
 
@@ -342,9 +342,9 @@ public:
 	}
 
 
-	void EnableMaterialID();
+	GEOMETRYCORE_API void EnableMaterialID();
 
-	void DisableMaterialID();
+	GEOMETRYCORE_API void DisableMaterialID();
 
 	FDynamicMeshMaterialAttribute* GetMaterialID()
 	{
@@ -360,10 +360,10 @@ public:
 
 	/// Create a new skin weights attribute with a given skin weights profile name. If an attribute already exists with
 	/// that name, that existing attribute will be deleted.
-	void AttachSkinWeightsAttribute(FName InProfileName, FDynamicMeshVertexSkinWeightsAttribute* InAttribute);
+	GEOMETRYCORE_API void AttachSkinWeightsAttribute(FName InProfileName, FDynamicMeshVertexSkinWeightsAttribute* InAttribute);
 
 	/// Remove a skin weights attribute matching the given profile name.
-	void RemoveSkinWeightsAttribute(FName InProfileName);
+	GEOMETRYCORE_API void RemoveSkinWeightsAttribute(FName InProfileName);
 
 	/// Returns true if the list of skin weight attributes includes the given profile name.
 	bool HasSkinWeightsAttribute(FName InProfileName) const
@@ -397,18 +397,18 @@ public:
 	// Bone Attributes
 	//
 
-	void CopyBoneAttributes(const FDynamicMeshAttributeSet& Copy);
+	GEOMETRYCORE_API void CopyBoneAttributes(const FDynamicMeshAttributeSet& Copy);
 	
-	void EnableMatchingBoneAttributes(const FDynamicMeshAttributeSet& ToMatch, bool bClearExisting, bool bDiscardExtraAttributes);
+	GEOMETRYCORE_API void EnableMatchingBoneAttributes(const FDynamicMeshAttributeSet& ToMatch, bool bClearExisting, bool bDiscardExtraAttributes);
 
 	/** @note Only compares bone names and parent indices. */
-	bool IsSameBoneAttributesAs(const FDynamicMeshAttributeSet& Other) const;
+	GEOMETRYCORE_API bool IsSameBoneAttributesAs(const FDynamicMeshAttributeSet& Other) const;
 
 	/**
 	 * The attribute is valid if either all attributes are empty or if the bone name attribute is not empty then all the other
 	 * attributes must be either empty or their size is equal to the bone name attribute size.
 	 */
-	bool CheckBoneValidity(EValidityCheckFailMode FailMode) const;
+	GEOMETRYCORE_API bool CheckBoneValidity(EValidityCheckFailMode FailMode) const;
 
 	
 	/** 
@@ -418,16 +418,16 @@ public:
 	 * 
 	 * @return true, if append was successful
 	 */
-	bool AppendBonesUnique(const FDynamicMeshAttributeSet& Other);
+	GEOMETRYCORE_API bool AppendBonesUnique(const FDynamicMeshAttributeSet& Other);
 
 	/** Enable all bone attributes and intialize their size to the given bone number. */
-	void EnableBones(const int InBonesNum);
+	GEOMETRYCORE_API void EnableBones(const int InBonesNum);
 
 	/** Disable all bone attributes. */
-	void DisableBones();
+	GEOMETRYCORE_API void DisableBones();
 	
 	/** Get number of bones. */
-	int32 GetNumBones() const;
+	GEOMETRYCORE_API int32 GetNumBones() const;
 
 	bool HasBones() const 
 	{ 
@@ -521,7 +521,7 @@ public:
 	/**
 	 * Returns true if this AttributeSet is the same as Other.
 	 */
-	bool IsSameAs(const FDynamicMeshAttributeSet& Other, bool bIgnoreDataLayout) const;
+	GEOMETRYCORE_API bool IsSameAs(const FDynamicMeshAttributeSet& Other, bool bIgnoreDataLayout) const;
 
 	/**
 	 * Serialization operator for FDynamicMeshAttributeSet.
@@ -543,7 +543,7 @@ public:
 	* @param CompactMaps If this is not a null pointer, the mesh serialization compacted the vertex and/or triangle data using the provided mapping. 
 	* @param bUseCompression Use compression for serializing bulk data.
 	*/
-	void Serialize(FArchive& Ar, const FCompactMaps* CompactMaps, bool bUseCompression);
+	GEOMETRYCORE_API void Serialize(FArchive& Ar, const FCompactMaps* CompactMaps, bool bUseCompression);
 
 protected:
 	/** Parent mesh of this attribute set */
@@ -592,17 +592,17 @@ protected:
 	// These functions are called by the FDynamicMesh3 to update the various
 	// attributes when the parent mesh topology has been modified.
 	// TODO: would it be better to register all the overlays and attributes with the base set and not overload these?  maybe!
-	virtual void OnNewTriangle(int TriangleID, bool bInserted);
-	virtual void OnNewVertex(int VertexID, bool bInserted);
-	virtual void OnRemoveTriangle(int TriangleID);
-	virtual void OnRemoveVertex(int VertexID);
-	virtual void OnReverseTriOrientation(int TriangleID);
-	virtual void OnSplitEdge(const DynamicMeshInfo::FEdgeSplitInfo & splitInfo);
-	virtual void OnFlipEdge(const DynamicMeshInfo::FEdgeFlipInfo & flipInfo);
-	virtual void OnCollapseEdge(const DynamicMeshInfo::FEdgeCollapseInfo & collapseInfo);
-	virtual void OnPokeTriangle(const DynamicMeshInfo::FPokeTriangleInfo & pokeInfo);
-	virtual void OnMergeEdges(const DynamicMeshInfo::FMergeEdgesInfo & mergeInfo);
-	virtual void OnSplitVertex(const DynamicMeshInfo::FVertexSplitInfo& SplitInfo, const TArrayView<const int>& TrianglesToUpdate);
+	GEOMETRYCORE_API virtual void OnNewTriangle(int TriangleID, bool bInserted);
+	GEOMETRYCORE_API virtual void OnNewVertex(int VertexID, bool bInserted);
+	GEOMETRYCORE_API virtual void OnRemoveTriangle(int TriangleID);
+	GEOMETRYCORE_API virtual void OnRemoveVertex(int VertexID);
+	GEOMETRYCORE_API virtual void OnReverseTriOrientation(int TriangleID);
+	GEOMETRYCORE_API virtual void OnSplitEdge(const DynamicMeshInfo::FEdgeSplitInfo & splitInfo);
+	GEOMETRYCORE_API virtual void OnFlipEdge(const DynamicMeshInfo::FEdgeFlipInfo & flipInfo);
+	GEOMETRYCORE_API virtual void OnCollapseEdge(const DynamicMeshInfo::FEdgeCollapseInfo & collapseInfo);
+	GEOMETRYCORE_API virtual void OnPokeTriangle(const DynamicMeshInfo::FPokeTriangleInfo & pokeInfo);
+	GEOMETRYCORE_API virtual void OnMergeEdges(const DynamicMeshInfo::FMergeEdgesInfo & mergeInfo);
+	GEOMETRYCORE_API virtual void OnSplitVertex(const DynamicMeshInfo::FVertexSplitInfo& SplitInfo, const TArrayView<const int>& TrianglesToUpdate);
 
 	/**
 	 * Check validity of attributes
@@ -610,7 +610,7 @@ protected:
 	 * @param bAllowNonmanifold Accept non-manifold topology as valid. Note that this should almost always be true for attributes; non-manifold overlays are generally valid.
 	 * @param FailMode Desired behavior if mesh is found invalid
 	 */
-	virtual bool CheckValidity(bool bAllowNonmanifold, EValidityCheckFailMode FailMode) const;
+	GEOMETRYCORE_API virtual bool CheckValidity(bool bAllowNonmanifold, EValidityCheckFailMode FailMode) const;
 };
 
 
diff --git a/Engine/Source/Runtime/GeometryCore/Public/DynamicMesh/DynamicMeshChangeTracker.h b/Engine/Source/Runtime/GeometryCore/Public/DynamicMesh/DynamicMeshChangeTracker.h
--- a/Engine/Source/Runtime/GeometryCore/Public/DynamicMesh/DynamicMeshChangeTracker.h
+++ b/Engine/Source/Runtime/GeometryCore/Public/DynamicMesh/DynamicMeshChangeTracker.h
@@ -110,7 +110,7 @@ typedef TDynamicMeshAttributeChange<float, 1> FDynamicMeshWeightChange;
 /**
  * FDynamicMeshAttributeChangeSet stores a set of UV and Normal changes for a FDynamicMesh3
  */
-class GEOMETRYCORE_API FDynamicMeshAttributeChangeSet
+class FDynamicMeshAttributeChangeSet
 {
 public:
 	UE_NONCOPYABLE(FDynamicMeshAttributeChangeSet);
@@ -126,7 +126,7 @@ public:
 	TArray<TUniquePtr<FDynamicMeshAttributeChangeBase>> RegisteredAttributeChanges;
 
 	/** call ::Apply() on all the UV and Normal changes */
-	bool Apply(FDynamicMeshAttributeSet* Attributes, bool bRevert) const;
+	GEOMETRYCORE_API bool Apply(FDynamicMeshAttributeSet* Attributes, bool bRevert) const;
 };
 
 
@@ -138,20 +138,20 @@ public:
  * Construction of a well-formed FDynamicMeshChange is quite complex and it is strongly
  * suggested that you do so via FDynamicMeshChangeTracker
  */
-class GEOMETRYCORE_API FDynamicMeshChange
+class FDynamicMeshChange
 {
 public:
-	~FDynamicMeshChange();
+	GEOMETRYCORE_API ~FDynamicMeshChange();
 
 	/** Store the initial state of a vertex */
-	void SaveInitialVertex(const FDynamicMesh3* Mesh, int VertexID);
+	GEOMETRYCORE_API void SaveInitialVertex(const FDynamicMesh3* Mesh, int VertexID);
 	/** Store the initial state of a triangle */
-	void SaveInitialTriangle(const FDynamicMesh3* Mesh, int TriangleID);
+	GEOMETRYCORE_API void SaveInitialTriangle(const FDynamicMesh3* Mesh, int TriangleID);
 
 	/** Store the final state of a vertex */
-	void StoreFinalVertex(const FDynamicMesh3* Mesh, int VertexID);
+	GEOMETRYCORE_API void StoreFinalVertex(const FDynamicMesh3* Mesh, int VertexID);
 	/** Store the final state of a triangle */
-	void StoreFinalTriangle(const FDynamicMesh3* Mesh, int TriangleID);
+	GEOMETRYCORE_API void StoreFinalTriangle(const FDynamicMesh3* Mesh, int TriangleID);
 
 	/** Attach an attribute change set to this mesh change, which will the be applied/reverted automatically */
 	void AttachAttributeChanges(TUniquePtr<FDynamicMeshAttributeChangeSet> AttribChanges)
@@ -160,19 +160,19 @@ public:
 	}
 
 	/** Apply or Revert this change using the given Mesh */
-	bool Apply(FDynamicMesh3* Mesh, bool bRevert) const;
+	GEOMETRYCORE_API bool Apply(FDynamicMesh3* Mesh, bool bRevert) const;
 
 	/** Do (limited) sanity checks on this MeshChange to ensure it is well-formed */
-	void VerifySaveState() const;
+	GEOMETRYCORE_API void VerifySaveState() const;
 
 	/** @return true if this vertex was saved. Uses linear search. */
-	bool HasSavedVertex(int VertexID);
+	GEOMETRYCORE_API bool HasSavedVertex(int VertexID);
 
 	/** store IDs of saved triangles in TrianglesOut. if bInitial=true, old triangles are stored, otherwise new triangles */
-	void GetSavedTriangleList(TArray<int>& TrianglesOut, bool bInitial) const;
+	GEOMETRYCORE_API void GetSavedTriangleList(TArray<int>& TrianglesOut, bool bInitial) const;
 
 	/** run self-validity checks on internal data structures to test if change is well-formed */
-	void CheckValidity(EValidityCheckFailMode FailMode = EValidityCheckFailMode::Check) const;
+	GEOMETRYCORE_API void CheckValidity(EValidityCheckFailMode FailMode = EValidityCheckFailMode::Check) const;
 
 protected:
 
@@ -198,7 +198,7 @@ protected:
 
 	TUniquePtr<FDynamicMeshAttributeChangeSet> AttributeChanges;
 
-	void ApplyReplaceChange(FDynamicMesh3* Mesh,
+	GEOMETRYCORE_API void ApplyReplaceChange(FDynamicMesh3* Mesh,
 		const TArray<FChangeTriangle>& RemoveTris,
 		const TArray<FChangeVertex>& InsertVerts,
 		const TArray<FChangeTriangle>& InsertTris) const;
@@ -216,27 +216,27 @@ protected:
  * objects (stored in a FDynamicMeshAttributeChangeSet). You should not use this class
  * directly, it is intended to be used via FDynamicMeshChangeTracker
  */
-class GEOMETRYCORE_API FDynamicMeshAttributeSetChangeTracker
+class FDynamicMeshAttributeSetChangeTracker
 {
 public:
-	explicit FDynamicMeshAttributeSetChangeTracker(const FDynamicMeshAttributeSet* Attribs);
+	GEOMETRYCORE_API explicit FDynamicMeshAttributeSetChangeTracker(const FDynamicMeshAttributeSet* Attribs);
 
 	/** Start tracking a change */
-	void BeginChange();
+	GEOMETRYCORE_API void BeginChange();
 	/** End the change transaction and get the resulting change object */
-	TUniquePtr<FDynamicMeshAttributeChangeSet> EndChange();
+	GEOMETRYCORE_API TUniquePtr<FDynamicMeshAttributeChangeSet> EndChange();
 
 	/** Store the initial state of a triangle */
-	void SaveInitialTriangle(int TriangleID);
+	GEOMETRYCORE_API void SaveInitialTriangle(int TriangleID);
 
 	/** store the final state of a set of triangles */
-	void StoreAllFinalTriangles(const TArray<int>& TriangleIDs);
+	GEOMETRYCORE_API void StoreAllFinalTriangles(const TArray<int>& TriangleIDs);
 
 	/** Store the initial state of a vertex */
-	void SaveInitialVertex(int VertexID);
+	GEOMETRYCORE_API void SaveInitialVertex(int VertexID);
 
 	/** store the final state of a set of vertices */
-	void StoreAllFinalVertices(const TArray<int>& VertexIDs);
+	GEOMETRYCORE_API void StoreAllFinalVertices(const TArray<int>& VertexIDs);
 
 
 protected:
@@ -286,19 +286,19 @@ protected:
  * @warning Currently only vertices that are part of saved triangles will be stored in the emitted FMeshChange!
  * 
  */
-class GEOMETRYCORE_API FDynamicMeshChangeTracker
+class FDynamicMeshChangeTracker
 {
 public:
-	explicit FDynamicMeshChangeTracker(const FDynamicMesh3* Mesh);
-	~FDynamicMeshChangeTracker();
+	GEOMETRYCORE_API explicit FDynamicMeshChangeTracker(const FDynamicMesh3* Mesh);
+	GEOMETRYCORE_API ~FDynamicMeshChangeTracker();
 
 	/** Initialize the change-tracking process */
-	void BeginChange();
+	GEOMETRYCORE_API void BeginChange();
 	/** Construct a change object that represents the delta between the Begin and End states */
-	TUniquePtr<FDynamicMeshChange> EndChange();
+	GEOMETRYCORE_API TUniquePtr<FDynamicMeshChange> EndChange();
 
 	/** Save necessary information about a triangle before it is modified */
-	void SaveTriangle(int32 TriangleID, bool bSaveVertices);
+	GEOMETRYCORE_API void SaveTriangle(int32 TriangleID, bool bSaveVertices);
 
 	/** Save necessary information about an edge before it is modified */
 	inline void SaveEdge(int32 EdgeID, bool bVertices);
@@ -319,7 +319,7 @@ public:
 	void SaveVertexOneRingTriangles(EnumerableType VertexIDs, bool bSaveVertices);
 
 	/** Do (limited) sanity checking to make sure that the change is well-formed*/
-	void VerifySaveState();
+	GEOMETRYCORE_API void VerifySaveState();
 
 protected:
 
@@ -330,7 +330,7 @@ protected:
 	//
 
 	/** Save necessary information about a vertex before it is modified */
-	void SaveVertex(int32 VertexID);
+	GEOMETRYCORE_API void SaveVertex(int32 VertexID);
 
 	/** Save necessary information about a set of vertices before they are modified */
 	template<typename EnumerableType>
@@ -425,4 +425,4 @@ void FDynamicMeshChangeTracker::SaveVertexOneRingTriangles(EnumerableType Vertex
 
 
 } // end namespace UE::Geometry
-} // end namespace UE
\ No newline at end of file
+} // end namespace UE
diff --git a/Engine/Source/Runtime/GeometryCore/Public/DynamicMesh/MeshIndexMappings.h b/Engine/Source/Runtime/GeometryCore/Public/DynamicMesh/MeshIndexMappings.h
--- a/Engine/Source/Runtime/GeometryCore/Public/DynamicMesh/MeshIndexMappings.h
+++ b/Engine/Source/Runtime/GeometryCore/Public/DynamicMesh/MeshIndexMappings.h
@@ -18,7 +18,7 @@ class FDynamicMesh3;
  * This is a convenient object to have, to avoid passing around large numbers of separate maps.
  * The individual maps are not necessarily all filled by every operation.
  */
-struct GEOMETRYCORE_API FMeshIndexMappings
+struct FMeshIndexMappings
 {
 protected:
 	FIndexMapi VertexMap;
@@ -31,7 +31,7 @@ protected:
 
 public:
 	/** Size internal arrays-of-maps to be suitable for this Mesh */
-	void Initialize(FDynamicMesh3* Mesh);
+	GEOMETRYCORE_API void Initialize(FDynamicMesh3* Mesh);
 
 	/** @return the value used to indicate "invalid" in the mapping */
 	constexpr int InvalidID() const { return VertexMap.UnmappedID(); }
@@ -93,4 +93,4 @@ public:
 };
 
 } // end namespace UE::Geometry
-} // end namespace UE
\ No newline at end of file
+} // end namespace UE
