Add missing Super::Tick call in AGameplayDebuggerPlayerManager::Tick

diff --git a/Engine/Source/Runtime/GameplayDebugger/Private/GameplayDebuggerPlayerManager.cpp b/Engine/Source/Runtime/GameplayDebugger/Private/GameplayDebuggerPlayerManager.cpp
--- a/Engine/Source/Runtime/GameplayDebugger/Private/GameplayDebuggerPlayerManager.cpp
+++ b/Engine/Source/Runtime/GameplayDebugger/Private/GameplayDebuggerPlayerManager.cpp
@@ -305,6 +305,8 @@ void AGameplayDebuggerPlayerManager::GetViewPoint(const APlayerController& Owner
 // FTickableGameObject begin
 void AGameplayDebuggerPlayerManager::Tick(float DeltaTime)
 {
+	Super::Tick(DeltaTime);
+	
 #if WITH_EDITORONLY_DATA 
 	if (EditorWorldData.Replicator)
 	{
