Iris
- Fixed bug in GameplayDebuggerCategoryReplicator when using Iris replication

diff --git a/Engine/Source/Runtime/GameplayDebugger/Private/GameplayDebuggerCategoryReplicator.cpp b/Engine/Source/Runtime/GameplayDebugger/Private/GameplayDebuggerCategoryReplicator.cpp
--- a/Engine/Source/Runtime/GameplayDebugger/Private/GameplayDebuggerCategoryReplicator.cpp
+++ b/Engine/Source/Runtime/GameplayDebugger/Private/GameplayDebuggerCategoryReplicator.cpp
@@ -122,7 +122,7 @@ void AGameplayDebuggerCategoryReplicator::PreReplication(IRepChangedPropertyTrac
 void AGameplayDebuggerCategoryReplicator::OnRep_ReplicatedData()
 {
 #if UE_WITH_IRIS
-	if (UE::Net::FReplicationSystemUtil::GetNetHandle(this).IsValid())
+	if (UE::Net::FReplicationSystemUtil::GetReplicationSystem(this))
 	{
 		ReplicatedData.ApplyToOwner();
 	}
