Fix crash with GameplayDebuggerCategoryReplicator when its held by ReplicationGraph. Setting it bReplicateToOwner should only have been done when Iris is actually running.

diff --git a/Engine/Source/Runtime/GameplayDebugger/Private/GameplayDebuggerCategoryReplicator.cpp b/Engine/Source/Runtime/GameplayDebugger/Private/GameplayDebuggerCategoryReplicator.cpp
--- a/Engine/Source/Runtime/GameplayDebugger/Private/GameplayDebuggerCategoryReplicator.cpp
+++ b/Engine/Source/Runtime/GameplayDebugger/Private/GameplayDebuggerCategoryReplicator.cpp
@@ -10,6 +10,7 @@
 #include "GameplayDebuggerRenderingComponent.h"
 #include "GameplayDebuggerExtension.h"
 #include "Net/UnrealNetwork.h"
+#include "Engine/NetDriver.h"
 #include "VisualLogger/VisualLogger.h"
 
 #if UE_WITH_IRIS
@@ -575,7 +576,17 @@ void AGameplayDebuggerCategoryReplicator::BeginPlay()
 #if UE_WITH_IRIS
 void AGameplayDebuggerCategoryReplicator::BeginReplication()
 {
-	bOnlyRelevantToOwner = true;
+	if (UWorld* World = GetWorld())
+	{
+		if (UNetDriver* NetDriver = World->GetNetDriver())
+		{
+			if (NetDriver->IsUsingIrisReplication())
+			{
+				bOnlyRelevantToOwner = true;
+			}
+		}
+	}
+
 	Super::BeginReplication();
 }
 
