Fix to GameplayDebugger's drawing in multiplayer PIE that addresses the debugging view of one player being drawn in all players' viewports

diff --git a/Engine/Source/Runtime/GameplayDebugger/Private/GameplayDebuggerLocalController.cpp b/Engine/Source/Runtime/GameplayDebugger/Private/GameplayDebuggerLocalController.cpp
--- a/Engine/Source/Runtime/GameplayDebugger/Private/GameplayDebuggerLocalController.cpp
+++ b/Engine/Source/Runtime/GameplayDebugger/Private/GameplayDebuggerLocalController.cpp
@@ -28,6 +28,7 @@
 #include "GameFramework/PlayerInput.h"
 #include "EngineUtils.h"
 #include "HAL/IConsoleManager.h"
+#include "SceneInterface.h"
 
 #include UE_INLINE_GENERATED_CPP_BY_NAME(GameplayDebuggerLocalController)
 
@@ -192,8 +193,18 @@ void UGameplayDebuggerLocalController::OnCategoriesChanged()
 }
 
 #if WITH_GAMEPLAY_DEBUGGER_MENU
-void UGameplayDebuggerLocalController::OnDebugDraw(class UCanvas* Canvas, class APlayerController* PC)
+void UGameplayDebuggerLocalController::OnDebugDraw(UCanvas* Canvas, APlayerController* PC)
 {
+	// this change is required for multi-client PIE, since even though every client has its own UWorld this OnDebugDraw
+	// gets called by a multicast-delegate - the same for all the clients. 
+	CA_SUPPRESS(6011);
+	const FSceneInterface* Scene = (Canvas && Canvas->Canvas) ? Canvas->Canvas->GetScene() : nullptr;
+	if (Scene && Scene->GetWorld() != CachedReplicator->GetWorld())
+	{
+		return;
+	}
+	check(Canvas);
+	
 	if (CachedReplicator && CachedReplicator->IsEnabled() && bDebugDrawEnabled)
 	{
 		FGameplayDebuggerCanvasContext CanvasContext(Canvas, HUDFont);
