Iris
- Fixed issue with GameplayDebuggerCategoryReplicator not replicating multilpe part DataPacks correctly

diff --git a/Engine/Source/Runtime/GameplayDebugger/Private/GameplayDebuggerCategoryReplicator.cpp b/Engine/Source/Runtime/GameplayDebugger/Private/GameplayDebuggerCategoryReplicator.cpp
--- a/Engine/Source/Runtime/GameplayDebugger/Private/GameplayDebuggerCategoryReplicator.cpp
+++ b/Engine/Source/Runtime/GameplayDebugger/Private/GameplayDebuggerCategoryReplicator.cpp
@@ -137,9 +137,9 @@ void FGameplayDebuggerNetPack::PopulateFromOwner()
 	// replicated systems so in order to make this work for Iris replication which has much stricter rules for what can be done during serialization,
 	// we cannot rely on polling being made during the call to the custom NetDeltaSerialize method so we use this explicit method to poll the data to be replicated instead.
 	
-	// Update SavedCategories for replication, DataPacks-packets are handled using RPC`s
 	if (Owner && Owner->bIsEnabled && Owner->Categories.Num() == SavedData.Num())
 	{
+		// Update SavedCategories for replication, DataPacks-packets are handled using RPC`s
 		for (int32 Idx = 0; Idx < SavedData.Num(); Idx++)
 		{
 			FGameplayDebuggerCategory& CategoryOb = Owner->Categories[Idx].Get();
@@ -159,6 +159,23 @@ void FGameplayDebuggerNetPack::PopulateFromOwner()
 			{
 				SavedCategory.Shapes = CategoryOb.ReplicatedShapes;
 			}
+
+			// Throttled send of DataPackRPC:s
+			if (Owner->bSendDataPacksUsingRPC)
+			{
+				for (int32 DataPackIdx = 0; DataPackIdx < CategoryOb.ReplicatedDataPacks.Num(); DataPackIdx++)
+				{
+					FGameplayDebuggerDataPack& DataPack = CategoryOb.ReplicatedDataPacks[DataPackIdx];
+					if (!CategoryOb.bIsLocal)
+					{
+						if ((DataPack.bNeedsConfirmation && !DataPack.bReceived))
+						{
+							// Send the update data pack as an reliable rpc instead
+							Owner->SendDataPackPacket(CategoryOb.GetCategoryName(), DataPackIdx, DataPack);						
+						}
+					}
+				}
+			}
 		}
 	}
 }
@@ -924,7 +941,7 @@ void AGameplayDebuggerCategoryReplicator::CollectCategoryData(bool bForce)
 						UE_LOG(LogGameplayDebugReplication, Verbose, TEXT("Category[%d].DataPack[%d] SENT, DataVersion:%d DataSize:%d SyncCounter:%d"),
 							Idx, DataPackIdx, DataPack.Header.DataVersion, DataPack.Header.DataSize, DataPack.Header.SyncCounter);
 
-						// Send the update data pack as an reliable rpc instead
+						// Send the update data pack as an reliable rpc instead, note this will just send the first part, multipart datapacks will be throttled over multiple frames
 						if (bSendDataPacksUsingRPC)
 						{
 							SendDataPackPacket(CategoryOb.GetCategoryName(), DataPackIdx, DataPack);						
