Gameplay Debugger now selects the local Actor on start-up if no other suitable Pawns are found (so you're likely to start with a valid Pawn -- the Player's Pawn).  You still won't be able to re-select if you choose another Pawn, unless you re-enable the Gameplay Debugger or use GDT.SelectLocalPlayer.

diff --git a/Engine/Source/Runtime/GameplayDebugger/Private/GameplayDebuggerLocalController.cpp b/Engine/Source/Runtime/GameplayDebugger/Private/GameplayDebuggerLocalController.cpp
--- a/Engine/Source/Runtime/GameplayDebugger/Private/GameplayDebuggerLocalController.cpp
+++ b/Engine/Source/Runtime/GameplayDebugger/Private/GameplayDebuggerLocalController.cpp
@@ -603,6 +603,11 @@ void UGameplayDebuggerLocalController::ToggleActivation()
 				GAreScreenMessagesEnabled = false;
 				DebugActorCandidate = nullptr;
 				OnSelectActorTick();
+
+				if (!DebugActorCandidate)
+				{
+					OnSelectLocalPlayer();
+				}
 			}
 			else
 			{
