Iris
*  Fixed bug with GameplayDebuggerCategoryReplicator when iris is globally enabled by default but overriden via the gamemode

#rnx

diff --git a/Engine/Source/Runtime/GameplayDebugger/Private/GameplayDebuggerCategoryReplicator.cpp b/Engine/Source/Runtime/GameplayDebugger/Private/GameplayDebuggerCategoryReplicator.cpp
--- a/Engine/Source/Runtime/GameplayDebugger/Private/GameplayDebuggerCategoryReplicator.cpp
+++ b/Engine/Source/Runtime/GameplayDebugger/Private/GameplayDebuggerCategoryReplicator.cpp
@@ -36,12 +36,6 @@ static TAutoConsoleVariable<int32> CVarGameplayDebuggerRepDetails(
 	TEXT("Enable or disable very verbose replication logs for gameplay debugger"),
 	ECVF_Cheat);
 
-static TAutoConsoleVariable<int32> CVarGameplayDebuggerUseDataPackRPC(
-	TEXT("GameplayDebugger.UseDataPackRPC"),
-	0,
-	TEXT("Enable or disable use of rpc's for datapack packets for gameplay debugger"),
-	ECVF_Cheat);
-
 FNotifyGameplayDebuggerOwnerChange AGameplayDebuggerCategoryReplicator::NotifyDebuggerOwnerChange;
 
 class FNetFastCategoryBaseState : public INetDeltaBaseState
@@ -562,19 +556,9 @@ AGameplayDebuggerCategoryReplicator::AGameplayDebuggerCategoryReplicator(const F
 	bIsEditorWorldReplicator = false;
 	bReplicates = true;
 
-#if UE_WITH_IRIS
-	if (UE::Net::ShouldUseIrisReplication())
-	{
-		// If iris is enabled, datapack-packets are always sent using RPC`s and we also requires PreReplication to called in order to populate replicated data outside of serialization.
-		SetCallPreReplication(true);
-		bSendDataPacksUsingRPC = true;
-	}
-#else
-	// We cache this as it does not make sense to be able to toggle it other than between sessions.
-	bSendDataPacksUsingRPC = CVarGameplayDebuggerUseDataPackRPC.GetValueOnAnyThread();
-#endif
-
 	ReplicatedData.Owner = this;
+
+	bSendDataPacksUsingRPC = false;
 }
 
 void AGameplayDebuggerCategoryReplicator::BeginPlay()
@@ -593,15 +577,12 @@ void AGameplayDebuggerCategoryReplicator::BeginPlay()
 #if UE_WITH_IRIS
 void AGameplayDebuggerCategoryReplicator::BeginReplication()
 {
-	if (UWorld* World = GetWorld())
+	if (UE::Net::FReplicationSystemUtil::GetReplicationSystem(this))
 	{
-		if (UNetDriver* NetDriver = World->GetNetDriver())
-		{
-			if (NetDriver->IsUsingIrisReplication())
-			{
-				bOnlyRelevantToOwner = true;
-			}
-		}
+		// If iris is enabled, datapack-packets are always sent using RPC`s and we also requires PreReplication to called in order to populate replicated data outside of serialization.
+		SetCallPreReplication(true);
+		bSendDataPacksUsingRPC = true;
+		bOnlyRelevantToOwner = true;
 	}
 
 	Super::BeginReplication();
