Small usability change for the Gameplay Debugger:  Always select the local Pawn.  Added instructions for selecting other Pawns.  This helps the case where you're testing multiclient games and your initial selection is confusing.

diff --git a/Engine/Source/Runtime/GameplayDebugger/Private/GameplayDebuggerLocalController.cpp b/Engine/Source/Runtime/GameplayDebugger/Private/GameplayDebuggerLocalController.cpp
--- a/Engine/Source/Runtime/GameplayDebugger/Private/GameplayDebuggerLocalController.cpp
+++ b/Engine/Source/Runtime/GameplayDebugger/Private/GameplayDebuggerLocalController.cpp
@@ -306,10 +306,15 @@ void UGameplayDebuggerLocalController::DrawHeader(FGameplayDebuggerCanvasContext
 	}
 	else
 	{
-		CanvasContext.Printf(TEXT("Tap {yellow}%s{white} to close, use %s to toggle categories."), *ActivationKeyDesc, *CategoryKeysDesc);
+		CanvasContext.Printf(TEXT("Tap {yellow}%s{white} to close or hold to select new Pawn. Use %s to toggle categories."), *ActivationKeyDesc, *CategoryKeysDesc);
 	}
 
-	const FString DebugActorDesc = FString::Printf(TEXT("Debug actor: {cyan}%s"), *CachedReplicator->GetDebugActorName().ToString());
+	// Get the NetRole string so we can hint if the user has selected a local Actor or not
+	const AActor* DebugActor = CachedReplicator ? CachedReplicator->GetDebugActor() : nullptr;
+	const ENetRole DebugNetRole = DebugActor ? DebugActor->GetLocalRole() : ENetRole::ROLE_None;
+	const FString DebugNetRoleString = UEnum::GetValueAsString<ENetRole>(DebugNetRole);
+	
+	const FString DebugActorDesc = FString::Printf(TEXT("Debug actor: {cyan}%s{white} [%s]"), *CachedReplicator->GetDebugActorName().ToString(), *DebugNetRoleString);
 	float DebugActorSizeX = 0.0f, DebugActorSizeY = 0.0f;
 	CanvasContext.MeasureString(DebugActorDesc, DebugActorSizeX, DebugActorSizeY);
 	CanvasContext.PrintAt((CanvasContext.Canvas->SizeX / DPIScale) - PaddingRight - DebugActorSizeX, UsePaddingTop, DebugActorDesc);
@@ -602,12 +607,7 @@ void UGameplayDebuggerLocalController::ToggleActivation()
 				bPrevScreenMessagesEnabled = GAreScreenMessagesEnabled;
 				GAreScreenMessagesEnabled = false;
 				DebugActorCandidate = nullptr;
-				OnSelectActorTick();
-
-				if (!DebugActorCandidate)
-				{
-					OnSelectLocalPlayer();
-				}
+				OnSelectLocalPlayer();
 			}
 			else
 			{
