[GameplayDebugger] switched back to best target selection on activation instead of using the local player

diff --git a/Engine/Source/Runtime/GameplayDebugger/Private/GameplayDebuggerLocalController.cpp b/Engine/Source/Runtime/GameplayDebugger/Private/GameplayDebuggerLocalController.cpp
--- a/Engine/Source/Runtime/GameplayDebugger/Private/GameplayDebuggerLocalController.cpp
+++ b/Engine/Source/Runtime/GameplayDebugger/Private/GameplayDebuggerLocalController.cpp
@@ -607,7 +607,13 @@ void UGameplayDebuggerLocalController::ToggleActivation()
 				bPrevScreenMessagesEnabled = GAreScreenMessagesEnabled;
 				GAreScreenMessagesEnabled = false;
 				DebugActorCandidate = nullptr;
-				OnSelectLocalPlayer();
+				OnSelectActorTick();
+
+				// If no actor got selected use local player
+				if (DebugActorCandidate == nullptr)
+				{
+					OnSelectLocalPlayer();
+				}
 			}
 			else
 			{
