[GameplayDebugger] adds key modifier 'Shift' support for local player selection
- <ActivationKey> : Toggle debugger ON + Single best pawn candidate selection (closest to reticle) | Toggle debugger OFF
- Hold <ActivationKey> : Toggle debugger ON (if not already opened) + Continuous best pawn candidate selection (closest to reticle)
- Shift + <ActivationKey> : Toggle debugger ON + Local player selection | Toggle debugger OFF
- Hold Shift + <ActivationKey> : Toggle debugger ON (if not already opened) + Local player selection

diff --git a/Engine/Source/Runtime/GameplayDebugger/Private/GameplayDebuggerLocalController.cpp b/Engine/Source/Runtime/GameplayDebugger/Private/GameplayDebuggerLocalController.cpp
--- a/Engine/Source/Runtime/GameplayDebugger/Private/GameplayDebuggerLocalController.cpp
+++ b/Engine/Source/Runtime/GameplayDebugger/Private/GameplayDebuggerLocalController.cpp
@@ -306,7 +306,11 @@ void UGameplayDebuggerLocalController::DrawHeader(FGameplayDebuggerCanvasContext
 	}
 	else
 	{
-		CanvasContext.Printf(TEXT("Tap {yellow}%s{white} to close or hold to select new Pawn. Use %s to toggle categories."), *ActivationKeyDesc, *CategoryKeysDesc);
+		const UGameplayDebuggerConfig* Config = UGameplayDebuggerConfig::StaticClass()->GetDefaultObject<UGameplayDebuggerConfig>();
+		CanvasContext.Printf(TEXT("Tap {yellow}%s{white} to close or hold to select new Pawn (hold {yellow}[Shift+%s]{white} to select local player). Use %s to toggle categories."),
+			*ActivationKeyDesc,
+			*Config->ActivationKey.ToString(),
+			*CategoryKeysDesc);
 	}
 
 	// Get the NetRole string so we can hint if the user has selected a local Actor or not
@@ -473,8 +477,11 @@ void UGameplayDebuggerLocalController::BindInput(UInputComponent& InputComponent
 	const UGameplayDebuggerConfig* SettingsCDO = UGameplayDebuggerConfig::StaticClass()->GetDefaultObject<UGameplayDebuggerConfig>();
 	if (!bSimulateMode)
 	{
-		InputComponent.BindKey(SettingsCDO->ActivationKey, IE_Pressed, this, &UGameplayDebuggerLocalController::OnActivationPressed);
-		InputComponent.BindKey(SettingsCDO->ActivationKey, IE_Released, this, &UGameplayDebuggerLocalController::OnActivationReleased);
+		InputComponent.BindKey(FInputChord(EModifierKey::None, SettingsCDO->ActivationKey), IE_Pressed, this, &UGameplayDebuggerLocalController::OnActivationPressed);
+		InputComponent.BindKey(FInputChord(EModifierKey::None, SettingsCDO->ActivationKey), IE_Released, this, &UGameplayDebuggerLocalController::OnActivationReleased);
+		InputComponent.BindKey(FInputChord(EModifierKey::Shift, SettingsCDO->ActivationKey), IE_Pressed, this, &UGameplayDebuggerLocalController::OnActivationPressedWithModifier);
+		InputComponent.BindKey(FInputChord(EModifierKey::Shift, SettingsCDO->ActivationKey), IE_Released, this, &UGameplayDebuggerLocalController::OnActivationReleasedWithModifier);
+
 		NewBindings.Add(SettingsCDO->ActivationKey.GetFName());
 	}
 
@@ -586,12 +593,28 @@ void UGameplayDebuggerLocalController::OnActivationPressed()
 	}
 }
 
+void UGameplayDebuggerLocalController::OnActivationPressedWithModifier()
+{
+	if (CachedReplicator)
+	{
+		//OnStartSelectingLocalPlayer();
+		const double HoldTimeThr = 0.2 * (FApp::UseFixedTimeStep() ? (FApp::GetFixedDeltaTime() * 60.) : 1.);
+		
+		CachedReplicator->GetWorldTimerManager().SetTimer(StartSelectingActorHandle, this, &UGameplayDebuggerLocalController::OnStartSelectingLocalPlayer, static_cast<float>(HoldTimeThr));
+	}
+}
+
 void UGameplayDebuggerLocalController::OnActivationReleased()
 {
-	ToggleActivation();
+	ToggleActivation(ESelectionMode::BestPawnCandidate);
+}
+
+void UGameplayDebuggerLocalController::OnActivationReleasedWithModifier()
+{
+	ToggleActivation(ESelectionMode::LocalPlayer);
 }
 
-void UGameplayDebuggerLocalController::ToggleActivation()
+void UGameplayDebuggerLocalController::ToggleActivation(const ESelectionMode SelectionMode)
 {
 	if (CachedReplicator)
 	{
@@ -607,7 +630,11 @@ void UGameplayDebuggerLocalController::ToggleActivation()
 				bPrevScreenMessagesEnabled = GAreScreenMessagesEnabled;
 				GAreScreenMessagesEnabled = false;
 				DebugActorCandidate = nullptr;
-				OnSelectActorTick();
+
+				if (SelectionMode == ESelectionMode::BestPawnCandidate)
+				{
+					OnSelectActorTick();
+				}
 
 				// If no actor got selected use local player
 				if (DebugActorCandidate == nullptr)
@@ -720,6 +747,16 @@ void UGameplayDebuggerLocalController::OnExtensionBindingEvent(int32 ExtensionId
 }
 
 void UGameplayDebuggerLocalController::OnStartSelectingActor()
+{
+	OnStartSelecting(ESelectionMode::BestPawnCandidate);
+}
+
+void UGameplayDebuggerLocalController::OnStartSelectingLocalPlayer()
+{
+	OnStartSelecting(ESelectionMode::LocalPlayer);
+}
+
+void UGameplayDebuggerLocalController::OnStartSelecting(ESelectionMode SelectionMode)
 {
 	StartSelectingActorHandle.Invalidate();
 	if (CachedReplicator)
@@ -736,10 +773,17 @@ void UGameplayDebuggerLocalController::OnStartSelectingActor()
 		bIsSelectingActor = true;
 		DebugActorCandidate = nullptr;
 
-		const bool bLooping = true;
-		CachedReplicator->GetWorldTimerManager().SetTimer(SelectActorTickHandle, this, &UGameplayDebuggerLocalController::OnSelectActorTick, 0.01f, bLooping);
+		if (SelectionMode == ESelectionMode::BestPawnCandidate)
+		{
+			const bool bLooping = true;
+			CachedReplicator->GetWorldTimerManager().SetTimer(SelectActorTickHandle, this, &UGameplayDebuggerLocalController::OnSelectActorTick, 0.01f, bLooping);
 
-		OnSelectActorTick();
+			OnSelectActorTick();
+		}
+		else if (SelectionMode == ESelectionMode::LocalPlayer)
+		{
+			OnSelectLocalPlayer();
+		}
 	}
 }
 
diff --git a/Engine/Source/Runtime/GameplayDebugger/Public/GameplayDebuggerLocalController.h b/Engine/Source/Runtime/GameplayDebugger/Public/GameplayDebuggerLocalController.h
--- a/Engine/Source/Runtime/GameplayDebugger/Public/GameplayDebuggerLocalController.h
+++ b/Engine/Source/Runtime/GameplayDebugger/Public/GameplayDebuggerLocalController.h
@@ -101,7 +101,9 @@ protected:
 	bool bDebugDrawEnabled = true;
 
 	void OnActivationPressed();
+	void OnActivationPressedWithModifier();
 	void OnActivationReleased();
+	void OnActivationReleasedWithModifier();
 	void OnCategory0Pressed();
 	void OnCategory1Pressed();
 	void OnCategory2Pressed();
@@ -117,6 +119,15 @@ protected:
 	void OnCategoryBindingEvent(int32 CategoryId, int32 HandlerId);
 	void OnExtensionBindingEvent(int32 ExtensionId, int32 HandlerId);
 
+	enum class ESelectionMode : uint8
+	{
+		BestPawnCandidate,
+		LocalPlayer
+	};
+
+	/** called short time after activation key was pressed and hold with 'Shift' */
+	void OnStartSelectingLocalPlayer();
+
 	/** sets the local player as the new debug actor */
 	void OnSelectLocalPlayer();
 
@@ -126,11 +137,13 @@ protected:
 	/** called in tick during actor selection */
 	void OnSelectActorTick();
 
+	void OnStartSelecting(ESelectionMode SelectionMode);
+	
 	/** toggle state of categories in given slot */
 	void ToggleSlotState(int32 SlotIdx);
 
 	/** toggle debugger on/off */
-	void ToggleActivation();
+	void ToggleActivation(ESelectionMode SelectionMode = ESelectionMode::BestPawnCandidate);
 
 	/** draw header row */
 	void DrawHeader(FGameplayDebuggerCanvasContext& CanvasContext);
