Improving ListenServer mode support for GameplayDebugger.  It still ends up rendering on the listen server, but now we at least get the replicated data on the non-authoritative clients.

diff --git a/Engine/Source/Runtime/GameplayDebugger/Private/GameplayDebuggerAddonManager.cpp b/Engine/Source/Runtime/GameplayDebugger/Private/GameplayDebuggerAddonManager.cpp
--- a/Engine/Source/Runtime/GameplayDebugger/Private/GameplayDebuggerAddonManager.cpp
+++ b/Engine/Source/Runtime/GameplayDebugger/Private/GameplayDebuggerAddonManager.cpp
@@ -123,7 +123,6 @@ void FGameplayDebuggerAddonManager::CreateCategories(AGameplayDebuggerCategoryRe
 	check(World);
 	const ENetMode NetMode = World->GetNetMode();
 	const bool bHasAuthority = FGameplayDebuggerUtils::IsAuthority(World);
-	const bool bIsLocal = (NetMode != NM_DedicatedServer);
 	const bool bIsSimulate = FGameplayDebuggerAddonBase::IsSimulateInEditor();
 
 	TArray<TSharedRef<FGameplayDebuggerCategory> > UnsortedCategories;
@@ -141,7 +140,7 @@ void FGameplayDebuggerAddonManager::CreateCategories(AGameplayDebuggerCategoryRe
 		CategoryObject.CategoryId = CategoryObjects.Num();
 		CategoryObject.CategoryName = It.Key;
 		CategoryObject.bHasAuthority = bHasAuthority;
-		CategoryObject.bIsLocal = bIsLocal;
+		CategoryObject.bIsLocal = Owner.IsLocal();
 		CategoryObject.bIsEnabled =
 			(It.Value.CategoryState == EGameplayDebuggerCategoryState::EnabledInGameAndSimulate) ||
 			(It.Value.CategoryState == EGameplayDebuggerCategoryState::EnabledInGame && !bIsSimulate) ||
diff --git a/Engine/Source/Runtime/GameplayDebugger/Private/GameplayDebuggerCategoryReplicator.cpp b/Engine/Source/Runtime/GameplayDebugger/Private/GameplayDebuggerCategoryReplicator.cpp
--- a/Engine/Source/Runtime/GameplayDebugger/Private/GameplayDebuggerCategoryReplicator.cpp
+++ b/Engine/Source/Runtime/GameplayDebugger/Private/GameplayDebuggerCategoryReplicator.cpp
@@ -569,7 +569,7 @@ void AGameplayDebuggerCategoryReplicator::BeginPlay()
 	check(World);
 	const ENetMode NetMode = World->GetNetMode();
 	bHasAuthority = FGameplayDebuggerUtils::IsAuthority(World);
-	bIsLocal = (NetMode != NM_DedicatedServer);
+	bIsLocal = !OwnerPC || OwnerPC->IsLocalPlayerController();
 
 	Init();
 }
