If the Gameplay Debugger is active while quitting, then restore the temporarily disabled GAreScreenMessagesEnabled flag (which allows messages to log to the screen).

Previously, quitting the game (e.g. ending PIE) while the Gameplay Debugger was active would prevent further messages from being logged to the screen in subsequent plays.

diff --git a/Engine/Source/Runtime/GameplayDebugger/Private/GameplayDebuggerLocalController.cpp b/Engine/Source/Runtime/GameplayDebugger/Private/GameplayDebuggerLocalController.cpp
--- a/Engine/Source/Runtime/GameplayDebugger/Private/GameplayDebuggerLocalController.cpp
+++ b/Engine/Source/Runtime/GameplayDebugger/Private/GameplayDebuggerLocalController.cpp
@@ -152,6 +152,12 @@ void UGameplayDebuggerLocalController::Cleanup()
 	}
 #endif // WITH_EDITOR
 
+	// If we are cleaning up while enabled, restore the screen messages flag
+	if (bIsLocallyEnabled && !GAreScreenMessagesEnabled)
+	{
+		GAreScreenMessagesEnabled = bPrevScreenMessagesEnabled;
+	}
+
 	bNeedsCleanup = false;
 }
 
