Massaged UGameplayDebuggerLocalController::OnDebugDraw to get rid of CA_SUPPRESS(6011) there, also made the code safer as a bonus.

diff --git a/Engine/Source/Runtime/GameplayDebugger/Private/GameplayDebuggerLocalController.cpp b/Engine/Source/Runtime/GameplayDebugger/Private/GameplayDebuggerLocalController.cpp
--- a/Engine/Source/Runtime/GameplayDebugger/Private/GameplayDebuggerLocalController.cpp
+++ b/Engine/Source/Runtime/GameplayDebugger/Private/GameplayDebuggerLocalController.cpp
@@ -201,65 +201,66 @@ void UGameplayDebuggerLocalController::OnCategoriesChanged()
 #if WITH_GAMEPLAY_DEBUGGER_MENU
 void UGameplayDebuggerLocalController::OnDebugDraw(UCanvas* Canvas, APlayerController* PC)
 {
-	// this change is required for multi-client PIE, since even though every client has its own UWorld this OnDebugDraw
-	// gets called by a multicast-delegate - the same for all the clients. 
-	CA_SUPPRESS(6011);
-	const FSceneInterface* Scene = (Canvas && Canvas->Canvas) ? Canvas->Canvas->GetScene() : nullptr;
+	check(Canvas);
+	if (CachedReplicator == nullptr || CachedReplicator->IsEnabled() == false || bDebugDrawEnabled == false)
+	{
+		return;
+	}
+
+	// we're checking if Scene's world is the same as CachedReplicator's since it can differ if we're debug-drawing
+	// in multiplayer PIE (OnDebugDraw is being called by a global multicast delegate, no world differentiation). 
+	const FSceneInterface* Scene = (Canvas->Canvas) ? Canvas->Canvas->GetScene() : nullptr;
 	if (Scene && Scene->GetWorld() != CachedReplicator->GetWorld())
 	{
 		return;
 	}
-	check(Canvas);
 	
-	if (CachedReplicator && CachedReplicator->IsEnabled() && bDebugDrawEnabled)
-	{
-		FGameplayDebuggerCanvasContext CanvasContext(Canvas, HUDFont);
-		CanvasContext.CursorX = CanvasContext.DefaultX = PaddingLeft;
-		CanvasContext.CursorY = CanvasContext.DefaultY = PaddingTop;
+	FGameplayDebuggerCanvasContext CanvasContext(Canvas, HUDFont);
+	CanvasContext.CursorX = CanvasContext.DefaultX = PaddingLeft;	
+	CanvasContext.CursorY = CanvasContext.DefaultY = PaddingTop;
 
-		CanvasContext.FontRenderInfo.bEnableShadow = bEnableTextShadow;
+	CanvasContext.FontRenderInfo.bEnableShadow = bEnableTextShadow;
 
-		APlayerController* ReplicationOwner = CachedReplicator->GetReplicationOwner();
-		CanvasContext.PlayerController = ReplicationOwner;
-		CanvasContext.World = ReplicationOwner ? ReplicationOwner->GetWorld() : CachedReplicator->GetWorld();
+	APlayerController* ReplicationOwner = CachedReplicator->GetReplicationOwner();
+	CanvasContext.PlayerController = ReplicationOwner;
+	CanvasContext.World = ReplicationOwner ? ReplicationOwner->GetWorld() : CachedReplicator->GetWorld();
 
-		DrawHeader(CanvasContext);
+	DrawHeader(CanvasContext);
 
-		if (DataPackMap.Num() != NumCategories)
-		{
-			RebuildDataPackMap();
-		}
+	if (DataPackMap.Num() != NumCategories)
+	{
+		RebuildDataPackMap();
+	}
 
-		if (Canvas->SceneView->ViewActor == nullptr)
-		{
-			CachedReplicator->SetViewPoint(Canvas->SceneView->ViewLocation, Canvas->SceneView->ViewRotation.Vector());
-		}
-		else if (CachedReplicator->IsViewPointSet())
-		{
-			CachedReplicator->ResetViewPoint();
-		}
+	if (Canvas->SceneView->ViewActor == nullptr)
+	{
+		CachedReplicator->SetViewPoint(Canvas->SceneView->ViewLocation, Canvas->SceneView->ViewRotation.Vector());
+	}
+	else if (CachedReplicator->IsViewPointSet())
+	{
+		CachedReplicator->ResetViewPoint();
+	}
 
-		const bool bHasDebugActor = CachedReplicator->HasDebugActor();
-		for (int32 Idx = 0; Idx < NumCategories; Idx++)
+	const bool bHasDebugActor = CachedReplicator->HasDebugActor();
+	for (int32 Idx = 0; Idx < NumCategories; Idx++)
+	{
+		TSharedRef<FGameplayDebuggerCategory> Category = CachedReplicator->GetCategory(Idx);
+		if (Category->ShouldDrawCategory(bHasDebugActor))
 		{
-			TSharedRef<FGameplayDebuggerCategory> Category = CachedReplicator->GetCategory(Idx);
-			if (Category->ShouldDrawCategory(bHasDebugActor))
+			// this is a special-case collection mode. If we want to collect data on the client this is the 
+			// place to do it, after the data got potentially replicated over from the server, and just 
+			// before drawing, so that new replicated data won't come in as we draw.
+			if ((Category->IsCategoryAuth() == false) && Category->ShouldCollectDataOnClient())
 			{
-				// this is a special-case collection mode. If we want to collect data on the client this is the 
-				// place to do it, after the data got potentially replicated over from the server, and just 
-				// before drawing, so that new replicated data won't come in as we draw.
-				if ((Category->IsCategoryAuth() == false) && Category->ShouldCollectDataOnClient())
-				{
-					Category->CollectData(ReplicationOwner, CachedReplicator->GetDebugActor());
-				}
-
-				if (Category->IsCategoryHeaderVisible())
-				{
-					DrawCategoryHeader(Idx, Category, CanvasContext);
-				}
+				Category->CollectData(ReplicationOwner, CachedReplicator->GetDebugActor());
+			}
 
-				Category->DrawCategory(ReplicationOwner, CanvasContext);
+			if (Category->IsCategoryHeaderVisible())
+			{
+				DrawCategoryHeader(Idx, Category, CanvasContext);
 			}
+
+			Category->DrawCategory(ReplicationOwner, CanvasContext);
 		}
 	}
 }
