GameplayDebugger: taking "Hidden" setting into consideration
#rnx

diff --git a/Engine/Source/Runtime/GameplayDebugger/Private/GameplayDebuggerAddonManager.cpp b/Engine/Source/Runtime/GameplayDebugger/Private/GameplayDebuggerAddonManager.cpp
--- a/Engine/Source/Runtime/GameplayDebugger/Private/GameplayDebuggerAddonManager.cpp
+++ b/Engine/Source/Runtime/GameplayDebugger/Private/GameplayDebuggerAddonManager.cpp
@@ -231,11 +231,14 @@ void FGameplayDebuggerAddonManager::UpdateFromConfig()
 			{
 				const bool bDefaultActiveInGame = (It.Value.DefaultCategoryState == EGameplayDebuggerCategoryState::EnabledInGame) || (It.Value.DefaultCategoryState == EGameplayDebuggerCategoryState::EnabledInGameAndSimulate);
 				const bool bDefaultActiveInSimulate = (It.Value.DefaultCategoryState == EGameplayDebuggerCategoryState::EnabledInSimulate) || (It.Value.DefaultCategoryState == EGameplayDebuggerCategoryState::EnabledInGameAndSimulate);
+				const bool bDefaultHidden = It.Value.DefaultCategoryState == EGameplayDebuggerCategoryState::Hidden;
 
 				const bool bActiveInGame = (ConfigData.ActiveInGame == EGameplayDebuggerOverrideMode::UseDefault) ? bDefaultActiveInGame : (ConfigData.ActiveInGame == EGameplayDebuggerOverrideMode::Enable);
 				const bool bActiveInSimulate = (ConfigData.ActiveInSimulate == EGameplayDebuggerOverrideMode::UseDefault) ? bDefaultActiveInSimulate : (ConfigData.ActiveInSimulate == EGameplayDebuggerOverrideMode::Enable);
+				const bool bHidden = (ConfigData.Hidden == EGameplayDebuggerOverrideMode::UseDefault) ? bDefaultHidden : (ConfigData.Hidden == EGameplayDebuggerOverrideMode::Enable);
 
 				EGameplayDebuggerCategoryState NewCategoryState =
+					bHidden ? EGameplayDebuggerCategoryState::Hidden :
 					bActiveInGame && bActiveInSimulate ? EGameplayDebuggerCategoryState::EnabledInGameAndSimulate :
 					bActiveInGame ? EGameplayDebuggerCategoryState::EnabledInGame :
 					bActiveInSimulate ? EGameplayDebuggerCategoryState::EnabledInSimulate :
