Added cvar to prevent gameplaydebuggermanager from automatically getting spawned

#rnx

diff --git a/Engine/Source/Runtime/GameplayDebugger/Private/GameplayDebuggerModule.cpp b/Engine/Source/Runtime/GameplayDebugger/Private/GameplayDebuggerModule.cpp
--- a/Engine/Source/Runtime/GameplayDebugger/Private/GameplayDebuggerModule.cpp
+++ b/Engine/Source/Runtime/GameplayDebugger/Private/GameplayDebuggerModule.cpp
@@ -18,6 +18,14 @@
 #include "Iris/IrisConfig.h"
 #endif // UE_WITH_IRIS
 
+
+static bool bAutoCreateGameplayDebuggerManager = true;
+static FAutoConsoleVariableRef CVarAutoCreateGameplayDebuggerManager(
+	TEXT("GameplayDebugger.AutoCreateGameplayDebuggerManager"),
+	bAutoCreateGameplayDebuggerManager,
+	TEXT("When false it stops auto-spawning the Gameplay Debugger Manager.")
+);
+
 FOnLocalControllerInitialized FGameplayDebuggerModule::OnLocalControllerInitialized;
 FOnLocalControllerUninitialized FGameplayDebuggerModule::OnLocalControllerUninitialized;
 
@@ -140,7 +148,7 @@ AGameplayDebuggerPlayerManager& FGameplayDebuggerModule::GetPlayerManager(UWorld
 void FGameplayDebuggerModule::OnWorldInitialized(UWorld* World, const UWorld::InitializationValues IVS)
 {
 	// make sure that world has valid player manager, create when it doesn't
-	if (World && (World->IsGameWorld() || World->IsEditorWorld()))
+	if (World && bAutoCreateGameplayDebuggerManager && (World->IsGameWorld() || World->IsEditorWorld()))
 	{
 		GetPlayerManager(World);
 	}
