Eliminate GROWABLE_PRINTF macro.

We can use a StringBuilder to implement the string formatting instead of copy-pasting the GROWABLE_LOGF macro.

diff --git a/Engine/Source/Runtime/GameplayDebugger/Private/GameplayDebuggerTypes.cpp b/Engine/Source/Runtime/GameplayDebugger/Private/GameplayDebuggerTypes.cpp
--- a/Engine/Source/Runtime/GameplayDebugger/Private/GameplayDebuggerTypes.cpp
+++ b/Engine/Source/Runtime/GameplayDebugger/Private/GameplayDebuggerTypes.cpp
@@ -9,6 +9,7 @@
 #include "GameplayDebuggerConfig.h"
 #include "DrawDebugHelpers.h"
 #include "CanvasItem.h"
+#include "Containers/StringFwd.h"
 #include "Engine/Canvas.h"
 #include "Engine/World.h"
 #include "GameFramework/PlayerController.h"
@@ -655,42 +656,20 @@ void FGameplayDebuggerCanvasContext::PrintAt(float PosX, float PosY, const FColo
 	Print(Color, Alpha, String);
 }
 
-// copied from Core/Private/Misc/VarargsHeler.h 
-#define GROWABLE_PRINTF(PrintFunc) \
-	int32	BufferSize	= 1024; \
-	TCHAR*	Buffer		= NULL; \
-	int32	Result		= -1; \
-	/* allocate some stack space to use on the first pass, which matches most strings */ \
-	TCHAR	StackBuffer[512]; \
-	TCHAR*	AllocatedBuffer = NULL; \
-\
-	/* first, try using the stack buffer */ \
-	Buffer = StackBuffer; \
-	GET_TYPED_VARARGS_RESULT( TCHAR, Buffer, UE_ARRAY_COUNT(StackBuffer), UE_ARRAY_COUNT(StackBuffer) - 1, Fmt, Fmt, Result ); \
-\
-	/* if that fails, then use heap allocation to make enough space */ \
-	while(Result == -1) \
-	{ \
-		FMemory::SystemFree(AllocatedBuffer); \
-		/* We need to use malloc here directly as GMalloc might not be safe. */ \
-		Buffer = AllocatedBuffer = (TCHAR*) FMemory::SystemMalloc( BufferSize * sizeof(TCHAR) ); \
-		GET_TYPED_VARARGS_RESULT( TCHAR, Buffer, BufferSize, BufferSize-1, Fmt, Fmt, Result ); \
-		BufferSize *= 2; \
-	}; \
-	Buffer[Result] = 0; \
-	; \
-\
-	PrintFunc; \
-	FMemory::SystemFree(AllocatedBuffer);
-
 void FGameplayDebuggerCanvasContext::PrintfImpl(const FColor& Color, const float Alpha, const TCHAR* Fmt, ...)
 {
-	GROWABLE_PRINTF(Print(Color, Alpha, Buffer));
+	va_list Args;
+	va_start(Args, Fmt);
+	Print(Color, Alpha, TStringBuilder<512>().AppendV(Fmt, Args).ToString());
+	va_end(Args);
 }
 
 void FGameplayDebuggerCanvasContext::PrintfAtImpl(float PosX, float PosY, const FColor& Color, const float Alpha, const TCHAR* Fmt, ...)
 {
-	GROWABLE_PRINTF(PrintAt(PosX, PosY, Color, Alpha, Buffer));
+	va_list Args;
+	va_start(Args, Fmt);
+	PrintAt(PosX, PosY, Color, Alpha, TStringBuilder<512>().AppendV(Fmt, Args).ToString());
+	va_end(Args);
 }
 
 void FGameplayDebuggerCanvasContext::MoveToNewLine()
