GameplayDebugger: which categories are enabled and disabled persist across PIE/editor/game sessions by saving them to EditorPerProjectUserSettings.ini by category name. Name is used because the category's id may change as modules are changed between sessions that register categories.

diff --git a/Engine/Source/Runtime/GameplayDebugger/Private/GameplayDebuggerLocalController.cpp b/Engine/Source/Runtime/GameplayDebugger/Private/GameplayDebuggerLocalController.cpp
--- a/Engine/Source/Runtime/GameplayDebugger/Private/GameplayDebuggerLocalController.cpp
+++ b/Engine/Source/Runtime/GameplayDebugger/Private/GameplayDebuggerLocalController.cpp
@@ -137,6 +137,9 @@ void UGameplayDebuggerLocalController::Initialize(AGameplayDebuggerCategoryRepli
 	AddonManager.OnCategoriesChanged.AddUObject(this, &UGameplayDebuggerLocalController::OnCategoriesChanged);
 	OnCategoriesChanged();
 
+	// Restore shown categories from last session
+	RestoreCategories();
+
 	bNeedsCleanup = true;
 }
 
@@ -198,6 +201,35 @@ void UGameplayDebuggerLocalController::OnCategoriesChanged()
 	DataPackMap.Reset();
 }
 
+void UGameplayDebuggerLocalController::RestoreCategories()
+{
+	// Retrieve from config which categories were enabled last session
+	const TSet<FName>& EnabledCategories = GetDefault<UGameplayDebuggerUserSettings>()->EnabledCategories;
+
+	// If saved category list is empty, use the category's default visibility. This is when 
+	// the user has never toggled a category or if they turned them all off.
+	if (EnabledCategories.IsEmpty())
+	{
+		return;
+	}
+
+	// Restore enabled and disabled categories from last session
+	const int32 NumReplicatorCategories = CachedReplicator->GetNumCategories();
+	for (int32 i = 0; i < NumReplicatorCategories; ++i)
+	{
+		// Get the category's name: this is saved to disk since id's are unstable
+		const TSharedRef<FGameplayDebuggerCategory>& Category = CachedReplicator->GetCategory(i);
+		const FName CategoryName = Category->GetCategoryName();
+
+		// Restore the category's enabled state
+		const bool bWasEnabledLastSession = EnabledCategories.Contains(CategoryName);
+		if (bWasEnabledLastSession != CachedReplicator->IsCategoryEnabled(i))
+		{
+			CachedReplicator->SetCategoryEnabled(i, bWasEnabledLastSession);
+		}
+	}
+}
+
 #if WITH_GAMEPLAY_DEBUGGER_MENU
 void UGameplayDebuggerLocalController::OnDebugDraw(UCanvas* Canvas, APlayerController* PC)
 {
@@ -856,8 +888,23 @@ void UGameplayDebuggerLocalController::ToggleSlotState(int32 SlotIdx)
 		const bool bIsEnabled = CachedReplicator->IsCategoryEnabled(SlotCategoryIds[SlotIdx][0]);
 		for (int32 Idx = 0; Idx < SlotCategoryIds[SlotIdx].Num(); Idx++)
 		{
+			// Inform server of category's new enabled state
+			const bool bNewEnabled = !bIsEnabled;
 			const int32 CategoryId = SlotCategoryIds[SlotIdx][Idx];
-			CachedReplicator->SetCategoryEnabled(CategoryId, !bIsEnabled);
+			CachedReplicator->SetCategoryEnabled(CategoryId, bNewEnabled);
+
+			// Save new enabled state to disk
+			const FName CategoryName = CachedReplicator->GetCategory(CategoryId)->GetCategoryName();
+			UGameplayDebuggerUserSettings* UserSettings = GetMutableDefault<UGameplayDebuggerUserSettings>();
+			if (bNewEnabled)
+			{
+				UserSettings->EnabledCategories.Add(CategoryName);
+			}
+			else
+			{
+				UserSettings->EnabledCategories.Remove(CategoryName);
+			}
+			UserSettings->SaveConfig();
 		}
 		// removed call to MarkComponentsRenderStateDirty since CachedReplicator->SetCategoryEnabled
 		// calls it already 
diff --git a/Engine/Source/Runtime/GameplayDebugger/Public/GameplayDebuggerConfig.h b/Engine/Source/Runtime/GameplayDebugger/Public/GameplayDebuggerConfig.h
--- a/Engine/Source/Runtime/GameplayDebugger/Public/GameplayDebuggerConfig.h
+++ b/Engine/Source/Runtime/GameplayDebugger/Public/GameplayDebuggerConfig.h
@@ -252,6 +252,12 @@ public:
 	UPROPERTY(config, EditAnywhere, Category = GameplayDebugger, meta = (UIMin = 0, ClampMin = 0, UIMax = 180, ClampMax = 180, Units = deg))
 	float MaxViewAngle = 45.f;
 
+	/** 
+	 * Persist which categories are enabled across editor sessions. Saved to disk as names since index is unstable. 
+	 */
+	UPROPERTY(config, EditAnywhere, Category = GameplayDebugger)
+	TSet<FName> EnabledCategories;
+
 protected:
 	/** Font Size used by Gameplay Debugger */ 
 	UPROPERTY(config, EditAnywhere, Category = GameplayDebugger)
diff --git a/Engine/Source/Runtime/GameplayDebugger/Public/GameplayDebuggerLocalController.h b/Engine/Source/Runtime/GameplayDebugger/Public/GameplayDebuggerLocalController.h
--- a/Engine/Source/Runtime/GameplayDebugger/Public/GameplayDebuggerLocalController.h
+++ b/Engine/Source/Runtime/GameplayDebugger/Public/GameplayDebuggerLocalController.h
@@ -168,4 +168,7 @@ protected:
 
 	/** called when known category set has changed */
 	void OnCategoriesChanged();
+
+	/** restore enabled categories from last session */
+	void RestoreCategories();
 };
\ No newline at end of file
