PR #13026: GameplayDebugger Local-Multiplayer Render-Duplication Fix

diff --git a/Engine/Source/Runtime/Engine/Classes/Debug/DebugDrawService.h b/Engine/Source/Runtime/Engine/Classes/Debug/DebugDrawService.h
--- a/Engine/Source/Runtime/Engine/Classes/Debug/DebugDrawService.h
+++ b/Engine/Source/Runtime/Engine/Classes/Debug/DebugDrawService.h
@@ -26,11 +26,10 @@ class UDebugDrawService : public UBlueprintFunctionLibrary
 	static ENGINE_API void Unregister(FDelegateHandle HandleToRemove);
 
 	/** Draws debug canvas that has already been initialized to a viewport */
-	static ENGINE_API void Draw(const FEngineShowFlags Flags, UCanvas* Canvas);
+	static ENGINE_API void Draw(const FEngineShowFlags Flags, UCanvas* Canvas, APlayerController* PlayerController = nullptr);
 
 	/** Initialize a debug canvas object then calls above draw. If CanvasObject is null it will find/create it for you */
-	static ENGINE_API void Draw(const FEngineShowFlags Flags, FViewport* Viewport, FSceneView* View, FCanvas* Canvas, UCanvas* CanvasObject = nullptr);
-
+	static ENGINE_API void Draw(const FEngineShowFlags Flags, FViewport* Viewport, FSceneView* View, FCanvas* Canvas, UCanvas* CanvasObject = nullptr, APlayerController* PlayerController = nullptr);
 private:
 
 	/**
diff --git a/Engine/Source/Runtime/Engine/Private/Debug/DebugDrawService.cpp b/Engine/Source/Runtime/Engine/Private/Debug/DebugDrawService.cpp
--- a/Engine/Source/Runtime/Engine/Private/Debug/DebugDrawService.cpp
+++ b/Engine/Source/Runtime/Engine/Private/Debug/DebugDrawService.cpp
@@ -55,7 +55,7 @@ void UDebugDrawService::Unregister(const FDelegateHandle HandleToRemove)
 	}	
 }
 
-void UDebugDrawService::Draw(const FEngineShowFlags Flags, FViewport* Viewport, FSceneView* View, FCanvas* Canvas, UCanvas* CanvasObject)
+void UDebugDrawService::Draw(const FEngineShowFlags Flags, FViewport* Viewport, FSceneView* View, FCanvas* Canvas, UCanvas* CanvasObject, APlayerController* PlayerController)
 {
 	if (CanvasObject == nullptr)
 	{
@@ -70,14 +70,14 @@ void UDebugDrawService::Draw(const FEngineShowFlags Flags, FViewport* Viewport,
 	// Canvas must be initialized every draw because the FCanvas passed in is on the stack in some scenarios.
 	CanvasObject->Init(View->UnscaledViewRect.Width(), View->UnscaledViewRect.Height(), View, Canvas);
 
-	CanvasObject->Update();	
+	CanvasObject->Update();
 	CanvasObject->SetView(View);
 
 	// PreRender the player's view.
-	Draw(Flags, CanvasObject);	
+	Draw(Flags, CanvasObject, PlayerController);
 }
 
-void UDebugDrawService::Draw(const FEngineShowFlags Flags, UCanvas* Canvas)
+void UDebugDrawService::Draw(const FEngineShowFlags Flags, UCanvas* Canvas, APlayerController* PlayerController)
 {
 	if (Canvas == nullptr || Canvas->Canvas == nullptr)
 	{
@@ -91,8 +91,7 @@ void UDebugDrawService::Draw(const FEngineShowFlags Flags, UCanvas* Canvas)
 
 		if (Flags.GetSingleFlag(FlagIndex) && ObservedFlags.GetSingleFlag(FlagIndex))
 		{
-			MulticastDelegate.Broadcast(Canvas, nullptr);
+			MulticastDelegate.Broadcast(Canvas, PlayerController);
 		}
 	}
 }
-
diff --git a/Engine/Source/Runtime/Engine/Private/GameViewportClient.cpp b/Engine/Source/Runtime/Engine/Private/GameViewportClient.cpp
--- a/Engine/Source/Runtime/Engine/Private/GameViewportClient.cpp
+++ b/Engine/Source/Runtime/Engine/Private/GameViewportClient.cpp
@@ -2016,10 +2016,10 @@ void UGameViewportClient::Draw(FViewport* InViewport, FCanvas* SceneCanvas)
 							}
 						}
 
-						if (DebugCanvas != NULL )
+						if (DebugCanvas != nullptr)
 						{
 							DebugCanvas->PushAbsoluteTransform(FTranslationMatrix(CanvasOrigin));
-							UDebugDrawService::Draw(ViewFamily.EngineShowFlags, InViewport, View, DebugCanvas, DebugCanvasObject);
+							UDebugDrawService::Draw(ViewFamily.EngineShowFlags, InViewport, View, DebugCanvas, DebugCanvasObject, PlayerController);
 							DebugCanvas->PopTransform();
 						}
 
diff --git a/Engine/Source/Runtime/GameplayDebugger/Private/GameplayDebuggerLocalController.cpp b/Engine/Source/Runtime/GameplayDebugger/Private/GameplayDebuggerLocalController.cpp
--- a/Engine/Source/Runtime/GameplayDebugger/Private/GameplayDebuggerLocalController.cpp
+++ b/Engine/Source/Runtime/GameplayDebugger/Private/GameplayDebuggerLocalController.cpp
@@ -231,7 +231,7 @@ void UGameplayDebuggerLocalController::RestoreCategories()
 }
 
 #if WITH_GAMEPLAY_DEBUGGER_MENU
-void UGameplayDebuggerLocalController::OnDebugDraw(UCanvas* Canvas, APlayerController* PC)
+void UGameplayDebuggerLocalController::OnDebugDraw(UCanvas* Canvas, APlayerController* PlayerController)
 {
 	check(Canvas);
 	if (CachedReplicator == nullptr || CachedReplicator->IsEnabled() == false || bDebugDrawEnabled == false)
@@ -239,16 +239,16 @@ void UGameplayDebuggerLocalController::OnDebugDraw(UCanvas* Canvas, APlayerContr
 		return;
 	}
 
-	// we're checking if Scene's world is the same as CachedReplicator's since it can differ if we're debug-drawing
-	// in multiplayer PIE (OnDebugDraw is being called by a global multicast delegate, no world differentiation). 
-	const FSceneInterface* Scene = (Canvas->Canvas) ? Canvas->Canvas->GetScene() : nullptr;
-	if (Scene && Scene->GetWorld() != CachedReplicator->GetWorld())
+	// Draw on canvas only when OnDebugDraw (bound to global multicast delegate; no world differentiation)
+	// is called without a specific player controller or with the one associated to our replicator.
+	// This validation is required for multiplayer PIE, including split-screen.
+	if (PlayerController != nullptr && CachedReplicator->GetReplicationOwner() != PlayerController)
 	{
 		return;
 	}
-	
+
 	FGameplayDebuggerCanvasContext CanvasContext(Canvas, HUDFont);
-	CanvasContext.CursorX = CanvasContext.DefaultX = PaddingLeft;	
+	CanvasContext.CursorX = CanvasContext.DefaultX = PaddingLeft;
 	CanvasContext.CursorY = CanvasContext.DefaultY = PaddingTop;
 
 	CanvasContext.FontRenderInfo.bEnableShadow = bEnableTextShadow;
