GameplayDebugger: no FString definitions in headers

diff --git a/Engine/Source/Runtime/GameplayDebugger/Private/GameplayDebuggerTypes.cpp b/Engine/Source/Runtime/GameplayDebugger/Private/GameplayDebuggerTypes.cpp
--- a/Engine/Source/Runtime/GameplayDebugger/Private/GameplayDebuggerTypes.cpp
+++ b/Engine/Source/Runtime/GameplayDebugger/Private/GameplayDebuggerTypes.cpp
@@ -26,6 +26,18 @@ namespace FGameplayDebuggerUtils
 	}
 }
 
+namespace FGameplayDebuggerCanvasStrings
+{
+	FString ColorNameInput = TEXT("white");
+	FString ColorNameEnabled = TEXT("green");
+	FString ColorNameDisabled = TEXT("grey");
+	FString ColorNameEnabledActiveRow = TEXT("green");
+	FString ColorNameDisabledActiveRow = TEXT("black");
+
+	FString Separator = TEXT("{white} | ");
+	FString SeparatorSpace = TEXT("  ");
+}
+
 //////////////////////////////////////////////////////////////////////////
 // FGameplayDebuggerShape
 
diff --git a/Engine/Source/Runtime/GameplayDebugger/Public/GameplayDebuggerTypes.h b/Engine/Source/Runtime/GameplayDebugger/Public/GameplayDebuggerTypes.h
--- a/Engine/Source/Runtime/GameplayDebugger/Public/GameplayDebuggerTypes.h
+++ b/Engine/Source/Runtime/GameplayDebugger/Public/GameplayDebuggerTypes.h
@@ -149,14 +149,14 @@ public:
 
 namespace FGameplayDebuggerCanvasStrings
 {
-	static FString ColorNameInput = TEXT("white");
-	static FString ColorNameEnabled = TEXT("green");
-	static FString ColorNameDisabled = TEXT("grey");
-	static FString ColorNameEnabledActiveRow = TEXT("green");
-	static FString ColorNameDisabledActiveRow = TEXT("black");
-
-	static FString Separator = TEXT("{white} | ");
-	static FString SeparatorSpace = TEXT("  ");
+	extern GAMEPLAYDEBUGGER_API FString ColorNameInput;
+	extern GAMEPLAYDEBUGGER_API FString ColorNameEnabled;
+	extern GAMEPLAYDEBUGGER_API FString ColorNameDisabled;
+	extern GAMEPLAYDEBUGGER_API FString ColorNameEnabledActiveRow;
+	extern GAMEPLAYDEBUGGER_API FString ColorNameDisabledActiveRow;
+
+	extern GAMEPLAYDEBUGGER_API FString Separator;
+	extern GAMEPLAYDEBUGGER_API FString SeparatorSpace;
 }
 
 UENUM()
