Some fixes for the GameplayDebugger:

1. The FCompression calls crashed in AutoRTFM, so now we're doing that path in the Open.
2. The Gameplay Debugger Replicator needs to override some Network-Role related functions to work with the MultiServerProxy.
3. The Gameplay Debugger Replicator used indices to replicate (which wouldn't necessarily match the other server). Used Names when available, Extensions are still by Index.

diff --git a/Engine/Source/Runtime/Core/Private/Misc/Compression.cpp b/Engine/Source/Runtime/Core/Private/Misc/Compression.cpp
--- a/Engine/Source/Runtime/Core/Private/Misc/Compression.cpp
+++ b/Engine/Source/Runtime/Core/Private/Misc/Compression.cpp
@@ -815,7 +815,11 @@ bool FCompression::CompressMemoryIfWorthDecompressing(FName FormatName, bool& bO
 bool FCompression::CompressMemory(FName FormatName, void* CompressedBuffer, int32& CompressedSize, const void* UncompressedBuffer, int32 UncompressedSize, ECompressionFlags Flags, int32 CompressionData)
 {
 	int64 CompressedSize64 = CompressedSize;
-	bool bSucceeded = CompressMemory(FormatName, CompressedBuffer, CompressedSize64, UncompressedBuffer, UncompressedSize, Flags, CompressionData);
+	bool bSucceeded = false;
+	UE_AUTORTFM_OPEN
+	{
+		bSucceeded = CompressMemory(FormatName, CompressedBuffer, CompressedSize64, UncompressedBuffer, UncompressedSize, Flags, CompressionData);
+	};
 	if (!IntFitsIn<int32>(CompressedSize64))
 	{
 		UE_LOG(LogCompression, Error, TEXT("Compressing 32 bit memory size ended up a 64 bit size! %d -> %lld"), UncompressedSize, CompressedSize64);
diff --git a/Engine/Source/Runtime/GameplayDebugger/Private/GameplayDebuggerCategoryReplicator.cpp b/Engine/Source/Runtime/GameplayDebugger/Private/GameplayDebuggerCategoryReplicator.cpp
--- a/Engine/Source/Runtime/GameplayDebugger/Private/GameplayDebuggerCategoryReplicator.cpp
+++ b/Engine/Source/Runtime/GameplayDebugger/Private/GameplayDebuggerCategoryReplicator.cpp
@@ -584,6 +584,29 @@ const AActor* AGameplayDebuggerCategoryReplicator::GetNetOwner() const
 	return OwnerPC.Get();
 }
 
+UPlayer* AGameplayDebuggerCategoryReplicator::GetNetOwningPlayer()
+{
+	if (const AActor* NetOwner = GetNetOwner())
+	{
+		if (NetOwner->GetLocalRole() == ROLE_Authority)
+		{
+			return GetNetOwningPlayerAnyRole();
+		}
+	}
+
+	return nullptr;
+}
+
+UPlayer* AGameplayDebuggerCategoryReplicator::GetNetOwningPlayerAnyRole()
+{
+	if (AActor* NetOwner = OwnerPC.Get())
+	{
+		return NetOwner->GetNetOwningPlayerAnyRole();
+	}
+
+	return nullptr;
+}
+
 #if WITH_EDITOR
 void AGameplayDebuggerCategoryReplicator::InitForEditor()
 {
@@ -737,9 +760,28 @@ bool AGameplayDebuggerCategoryReplicator::ServerSetCategoryEnabled_Validate(int3
 
 void AGameplayDebuggerCategoryReplicator::ServerSetCategoryEnabled_Implementation(int32 CategoryId, bool bEnable)
 {
+	UE_LOG(LogGameplayDebugReplication, Error, TEXT("%hs is Deprecated. Use ServerSetCategoryNameEnabled"), __func__);
 	SetCategoryEnabled(CategoryId, bEnable);
 }
 
+bool AGameplayDebuggerCategoryReplicator::ServerSetCategoryNameEnabled_Validate(FName CategoryName, bool bEnable)
+{
+	return true;
+}
+
+void AGameplayDebuggerCategoryReplicator::ServerSetCategoryNameEnabled_Implementation(FName CategoryName, bool bEnable)
+{
+	const int32 CategoryIndex = Categories.IndexOfByPredicate([CategoryName](const TSharedRef<FGameplayDebuggerCategory>& Item)
+	{
+		return Item->GetCategoryName() == CategoryName;
+	});
+
+	if (CategoryIndex != INDEX_NONE)
+	{
+		SetCategoryEnabled(CategoryIndex, bEnable);
+	}
+}
+
 bool AGameplayDebuggerCategoryReplicator::ServerSendCategoryInputEvent_Validate(int32 CategoryId, int32 HandlerId)
 {
 	return true;
@@ -747,9 +789,28 @@ bool AGameplayDebuggerCategoryReplicator::ServerSendCategoryInputEvent_Validate(
 
 void AGameplayDebuggerCategoryReplicator::ServerSendCategoryInputEvent_Implementation(int32 CategoryId, int32 HandlerId)
 {
+	UE_LOG(LogGameplayDebugReplication, Error, TEXT("%hs is Deprecated. Use ServerSendCategoryNameInputEvent"), __func__);
 	SendCategoryInputEvent(CategoryId, HandlerId);
 }
 
+bool AGameplayDebuggerCategoryReplicator::ServerSendCategoryNameInputEvent_Validate(FName CategoryName, int32 HandlerId)
+{
+	return true;
+}
+
+void AGameplayDebuggerCategoryReplicator::ServerSendCategoryNameInputEvent_Implementation(FName CategoryName, int32 HandlerId)
+{
+	const int32 CategoryIndex = Categories.IndexOfByPredicate([CategoryName](const TSharedRef<FGameplayDebuggerCategory>& Item)
+	{
+		return Item->GetCategoryName() == CategoryName;
+	});
+
+	if (CategoryIndex != INDEX_NONE)
+	{
+		SendCategoryInputEvent(CategoryIndex, HandlerId);
+	}	
+}
+
 bool AGameplayDebuggerCategoryReplicator::ServerSendExtensionInputEvent_Validate(int32 ExtensionId, int32 HandlerId)
 {
 	return true;
@@ -1087,17 +1148,20 @@ void AGameplayDebuggerCategoryReplicator::ResetViewPoint()
 
 void AGameplayDebuggerCategoryReplicator::SetCategoryEnabled(int32 CategoryId, bool bEnable)
 {
+	if (!Categories.IsValidIndex(CategoryId))
+	{
+		return;
+	}
+
 	if (bHasAuthority)
 	{
-		if (Categories.IsValidIndex(CategoryId))
-		{
-			UE_VLOG_UELOG(this, LogGameplayDebugReplication, Log, TEXT("SetCategoryEnabled[%d]:%d (%s)"), CategoryId, bEnable ? 1 : 0, *Categories[CategoryId]->GetCategoryName().ToString());
-			Categories[CategoryId]->bIsEnabled = bEnable;
-		}
+		UE_VLOG_UELOG(this, LogGameplayDebugReplication, Log, TEXT("SetCategoryEnabled[%d]:%d (%s)"), CategoryId, bEnable ? 1 : 0, *Categories[CategoryId]->GetCategoryName().ToString());
+		Categories[CategoryId]->bIsEnabled = bEnable;
 	}
 	else
 	{
-		ServerSetCategoryEnabled(CategoryId, bEnable);
+		FName CategoryName = Categories[CategoryId]->GetCategoryName();
+		ServerSetCategoryNameEnabled(CategoryName, bEnable);
 	}
 
 	MarkComponentsRenderStateDirty();
@@ -1118,7 +1182,7 @@ void AGameplayDebuggerCategoryReplicator::SendCategoryInputEvent(int32 CategoryI
 			}
 			else
 			{
-				ServerSendCategoryInputEvent(CategoryId, HandlerId);
+				ServerSendCategoryNameInputEvent(Categories[CategoryId]->GetCategoryName(), HandlerId);
 			}
 		}
 	}
diff --git a/Engine/Source/Runtime/GameplayDebugger/Public/GameplayDebuggerCategoryReplicator.h b/Engine/Source/Runtime/GameplayDebugger/Public/GameplayDebuggerCategoryReplicator.h
--- a/Engine/Source/Runtime/GameplayDebugger/Public/GameplayDebuggerCategoryReplicator.h
+++ b/Engine/Source/Runtime/GameplayDebugger/Public/GameplayDebuggerCategoryReplicator.h
@@ -14,6 +14,7 @@ class APlayerController;
 class FGameplayDebuggerCategory;
 class FGameplayDebuggerExtension;
 class UGameplayDebuggerRenderingComponent;
+class UPlayer;
 
 // Struct used to send the DataPackPackets as RPC`s instead of using the CustomDeltaNetSerializer.
 USTRUCT()
@@ -128,6 +129,8 @@ protected:
 	GAMEPLAYDEBUGGER_API virtual void OnReplicationStartedForIris(const FOnReplicationStartedParams&) override;
 
 	GAMEPLAYDEBUGGER_API virtual const AActor* GetNetOwner() const override;
+	GAMEPLAYDEBUGGER_API virtual UPlayer* GetNetOwningPlayer() override;
+	GAMEPLAYDEBUGGER_API virtual UPlayer* GetNetOwningPlayerAnyRole() override;
 
 public:
 	GAMEPLAYDEBUGGER_API virtual void EndPlay(const EEndPlayReason::Type EndPlayReason) override;
@@ -280,13 +283,21 @@ protected:
 	UFUNCTION(Server, Reliable, WithValidation, meta = (CallInEditor = "true"))
 	GAMEPLAYDEBUGGER_API void ServerResetViewPoint();
 
+	UE_DEPRECATED(5.8, "ServerSetCategoryEnabled by CategoryId does not work when server/client mismatch indices, use ServerSetCategoryNameEnabled")
 	UFUNCTION(Server, Reliable, WithValidation, meta = (CallInEditor = "true"))
 	GAMEPLAYDEBUGGER_API void ServerSetCategoryEnabled(int32 CategoryId, bool bEnable);
 
+	UFUNCTION(Server, Reliable, WithValidation, meta = (CallInEditor = "true"))
+	GAMEPLAYDEBUGGER_API void ServerSetCategoryNameEnabled(FName CategoryName, bool bEnable);
+
 	/** helper function for replicating input for category handlers */
+	UE_DEPRECATED(5.8, "ServerSendCategoryInputEvent by CategoryId does not work when server/client mismatch indices, use ServerSendCategoryNameInputEvent")
 	UFUNCTION(Server, Reliable, WithValidation, meta = (CallInEditor = "true"))
 	GAMEPLAYDEBUGGER_API void ServerSendCategoryInputEvent(int32 CategoryId, int32 HandlerId);
 
+	UFUNCTION(Server, Reliable, WithValidation, meta = (CallInEditor = "true"))
+	GAMEPLAYDEBUGGER_API void ServerSendCategoryNameInputEvent(FName CategoryName, int32 HandlerId);
+
 	/** helper function for replicating input for extension handlers */
 	UFUNCTION(Server, Reliable, WithValidation, meta = (CallInEditor = "true"))
 	GAMEPLAYDEBUGGER_API void ServerSendExtensionInputEvent(int32 ExtensionId, int32 HandlerId);
