We've seen some occurence of crash where category index is not valid out of the map, it should not crash but ensure.

slack thread: https://epic.enterprise.slack.com/archives/C0448GCBM5L/p1772135852057899?thread_ts=1772096134.431719&cid=C0448GCBM5L

#Synced-CL 51097280

diff --git a/Engine/Source/Runtime/GameplayDebugger/Private/GameplayDebuggerAddonManager.cpp b/Engine/Source/Runtime/GameplayDebugger/Private/GameplayDebuggerAddonManager.cpp
--- a/Engine/Source/Runtime/GameplayDebugger/Private/GameplayDebuggerAddonManager.cpp
+++ b/Engine/Source/Runtime/GameplayDebugger/Private/GameplayDebuggerAddonManager.cpp
@@ -158,7 +158,10 @@ void FGameplayDebuggerAddonManager::CreateCategories(AGameplayDebuggerCategoryRe
 		for (int32 Idx = 0; Idx < SlotMap[SlotIdx].Num(); Idx++)
 		{
 			const int32 CategoryId = SlotMap[SlotIdx][Idx];
-			CategoryObjects.Add(UnsortedCategories[CategoryId]);
+			if (ensureMsgf(UnsortedCategories.IsValidIndex(CategoryId), TEXT("Invalid Category Id from SlotMap - Could not add the Gameplay Debugger Category")))
+			{
+				CategoryObjects.Add(UnsortedCategories[CategoryId]);
+			}
 		}
 	}
 }
