Fix linker error with MergedModules + Test Editor builds, as bConsoleCommandsEnabled gets moved to a Common DLL while referenced from other DLLs causing a linker issue

#Synced-CL 51322414

#Synced-CL 51322414

diff --git a/Engine/Source/Runtime/GameplayDebugger/Public/GameplayDebuggerLocalController.h b/Engine/Source/Runtime/GameplayDebugger/Public/GameplayDebuggerLocalController.h
--- a/Engine/Source/Runtime/GameplayDebugger/Public/GameplayDebuggerLocalController.h
+++ b/Engine/Source/Runtime/GameplayDebugger/Public/GameplayDebuggerLocalController.h
@@ -75,7 +75,7 @@ protected:
 	int32 NumCategorySlots;
 	int32 NumCategories;
 
-	static bool bConsoleCommandsEnabled;
+	static GAMEPLAYDEBUGGER_API bool bConsoleCommandsEnabled;
 
 #if WITH_GAMEPLAY_DEBUGGER_MENU
 	friend struct FGameplayDebuggerConsoleCommands;
